implement a News Feed feature for the Level Editor (TTP# 331113)

diff --git a/Engine/Config/BaseEditor.ini b/Engine/Config/BaseEditor.ini
--- a/Engine/Config/BaseEditor.ini
+++ b/Engine/Config/BaseEditor.ini
@@ -1,8 +1,5 @@
 [Internationalization]
 +LocalizationPaths=../../../Engine/Content/Localization/Editor
-+ToolTipLocalizationPaths=../../../Engine/Content/Localization/Editor/ToolTips
-+PropertyNameLocalizationPaths=../../../Engine/Content/Localization/Editor/PropertyNames
-ShouldLoadLocalizedPropertyNames=True
 
 [UnrealEd.GeneralSettings]
 EnableStylusPressure=True
@@ -99,9 +96,7 @@ DefaultProfile="StarterProfile"
 +StarterProfile=(NodeName="/Script/BlueprintGraph.K2Node_CallFunction", FieldName="/Script/Engine.GameplayStatics:GetAllActorsOfClass")
 +StarterProfile=(NodeName="/Script/BlueprintGraph.K2Node_SpawnActorFromClass")
 
-[UnrealEd.MeshPaint]
-MinBrushRadius=1
-MaxBrushRadius=256
+
 
 [UnrealEd.HitProxy]
 HitProxySize=5
@@ -249,11 +244,7 @@ SupportWebsiteURL="http://unrealengine.com/support"
 [UserDefinedStructure]
 bUseUserDefinedStructure=false
 
-[CustomBlueprintFunctionLibrary]
-bCanCreateNew=true
-
-[EditoronlyBP]
-bAllowClassAndBlueprintPinMatching=false
-bReplaceBlueprintWithClass=false
-bDontLoadBlueprintOutsideEditor=false
-bBlueprintIsNotBlueprintType=false
\ No newline at end of file
+[NewsFeed]
+CdnUrl="http://cdn.unrealengine.com/newsfeed/{Culture}/index.json"
+EngineContentRelativeDirectory="Editor/NewsFeed/{Culture}/"
+Source="CDN"
diff --git a/Engine/Source/Editor/EditorStyle/Private/SlateEditorStyle.cpp b/Engine/Source/Editor/EditorStyle/Private/SlateEditorStyle.cpp
--- a/Engine/Source/Editor/EditorStyle/Private/SlateEditorStyle.cpp
+++ b/Engine/Source/Editor/EditorStyle/Private/SlateEditorStyle.cpp
@@ -3581,6 +3581,14 @@ void FSlateEditorStyle::FStyle::SetupLevelEditorStyle()
 		Set( "MeshPaint.Swap", new IMAGE_BRUSH( "/Icons/icon_MeshPaint_Swap_16x", Icon16x16 ) );
 	}
 
+	// News Feed
+	{
+		Set( "NewsFeed.ToolbarIcon.Small", new IMAGE_BRUSH( "NewsFeed/ToolbarIcon_16x", Icon16x16 ) );
+		Set( "NewsFeed.ReloadButton", new IMAGE_BRUSH( "NewsFeed/ReloadButton", Icon16x16 ) );
+		Set( "NewsFeed.SettingsButton", new IMAGE_BRUSH( "NewsFeed/SettingsButton", Icon16x16 ) );
+		Set( "NewsFeed.UnreadCountBackground", new IMAGE_BRUSH( "NewsFeed/UnreadCountBackground", Icon16x16 ) );
+	}
+
 	// EditorModesToolbar
 	{
 		Set( "EditorModesToolbar.Background", new FSlateNoResource() );
diff --git a/Engine/Source/Editor/LevelEditor/LevelEditor.Build.cs b/Engine/Source/Editor/LevelEditor/LevelEditor.Build.cs
--- a/Engine/Source/Editor/LevelEditor/LevelEditor.Build.cs
+++ b/Engine/Source/Editor/LevelEditor/LevelEditor.Build.cs
@@ -32,6 +32,7 @@ public class LevelEditor : ModuleRules
                 "EditorStyle",
 				"Engine",
 				"MessageLog",
+				"NewsFeed",
 				"Settings",
                 "SourceControl",
                 "StatsViewer",
diff --git a/Engine/Source/Editor/LevelEditor/Private/SLevelEditor.cpp b/Engine/Source/Editor/LevelEditor/Private/SLevelEditor.cpp
--- a/Engine/Source/Editor/LevelEditor/Private/SLevelEditor.cpp
+++ b/Engine/Source/Editor/LevelEditor/Private/SLevelEditor.cpp
@@ -36,13 +36,17 @@
 #include "EditorModes.h"
 #include "STutorialWrapper.h"
 #include "IDocumentation.h"
+#include "NewsFeed.h"
+
 
 static const FName LevelEditorBuildAndSubmitTab("LevelEditorBuildAndSubmit");
 static const FName LevelEditorStatsViewerTab("LevelEditorStatsViewer");
 static const FName LevelEditorWorldBrowserTab("LevelEditorWorldBrowser");
 static const FName MainFrameModuleName("MainFrame");
+static const FName NewsFeedModuleName("NewsFeed");
 static const FName LevelEditorModuleName("LevelEditor");
 
+
 namespace LevelEditorConstants
 {
 	/** The size of the thumbnail pool */
@@ -213,22 +217,36 @@ void SLevelEditor::Initialize( const TSharedRef<SDockTab>& OwnerTab, const TShar
 
 void SLevelEditor::ConstructNotificationBar()
 {
-	const IMainFrameModule& MainFrameModule = FModuleManager::GetModuleChecked< IMainFrameModule >( MainFrameModuleName );
-
 	NotificationBarBox->ClearChildren();
 
+	// level editor commands
 	NotificationBarBox->AddSlot()
-	.Padding( 5, 0 )
-	.AutoWidth()
-	[
-		FLevelEditorMenu::MakeNotificationBar( LevelEditorCommands, SharedThis(this ) )
-	];
+		.AutoWidth()
+		.Padding(5.0f, 0.0f, 0.0f, 0.0f)
+		[
+			FLevelEditorMenu::MakeNotificationBar( LevelEditorCommands, SharedThis(this ) )
+		];
+
+	// news feed button
+	INewsFeedModule& NewsFeedModule = FModuleManager::LoadModuleChecked<INewsFeedModule>(NewsFeedModuleName);
 
 	NotificationBarBox->AddSlot()
-	.AutoWidth()
-	[
-		MainFrameModule.MakeDeveloperTools()
-	];
+		.AutoWidth()
+		.Padding(5.0f, 0.0f, 0.0f, 0.0f)
+		.VAlign(VAlign_Bottom)
+		[
+			NewsFeedModule.CreateNewsFeedButton()
+		];
+
+	// developer tools
+	const IMainFrameModule& MainFrameModule = FModuleManager::GetModuleChecked<IMainFrameModule>(MainFrameModuleName);
+
+	NotificationBarBox->AddSlot()
+		.AutoWidth()
+		.Padding(5.0f, 0.0f, 0.0f, 0.0f)
+		[
+			MainFrameModule.MakeDeveloperTools()
+		];
 }
 
 
diff --git a/Engine/Source/Editor/NewsFeed/NewsFeed.Build.cs b/Engine/Source/Editor/NewsFeed/NewsFeed.Build.cs
new file mode 100644
--- /dev/null
+++ b/Engine/Source/Editor/NewsFeed/NewsFeed.Build.cs
@@ -0,0 +1,37 @@
+// Copyright 1998-2014 Epic Games, Inc. All Rights Reserved.
+
+using UnrealBuildTool;
+
+public class NewsFeed : ModuleRules
+{
+	public NewsFeed(TargetInfo Target)
+	{
+		PrivateIncludePaths.AddRange(
+			new string[]
+			{
+				"Editor/NewsFeed/Private",
+				"Editor/NewsFeed/Private/Implementation",
+				"Editor/NewsFeed/Private/Models",
+				"Editor/NewsFeed/Private/Widgets",
+			}
+		);
+
+		PrivateDependencyModuleNames.AddRange(
+			new string[]
+			{
+				"Analytics",
+				"Core",
+				"CoreUObject",
+				"DesktopPlatform",
+				"EditorStyle",
+				"Engine",
+				"HTTP",
+				"ImageWrapper",
+                "InputCore",
+				"OnlineSubsystem",
+				"Settings",
+				"Slate",
+			}
+		);
+	}
+}
diff --git a/Engine/Source/Editor/NewsFeed/Private/Implementation/CdnNewsFeedTitleFile.cpp b/Engine/Source/Editor/NewsFeed/Private/Implementation/CdnNewsFeedTitleFile.cpp
new file mode 100644
--- /dev/null
+++ b/Engine/Source/Editor/NewsFeed/Private/Implementation/CdnNewsFeedTitleFile.cpp
@@ -0,0 +1,400 @@
+// Copyright 1998-2014 Epic Games, Inc. All Rights Reserved.
+
+#include "NewsFeedPrivatePCH.h"
+
+
+DEFINE_LOG_CATEGORY_STATIC(LogEpicStorage, Display, All);
+
+IOnlineTitleFilePtr FCdnNewsFeedTitleFile::Create( const FString& IndexUrl )
+{
+	return MakeShareable( new FCdnNewsFeedTitleFile( IndexUrl ) );
+}
+
+FCdnNewsFeedTitleFile::FCdnNewsFeedTitleFile( const FString& InIndexUrl )
+	: IndexUrl( InIndexUrl )
+{
+
+}
+
+bool FCdnNewsFeedTitleFile::GetFileContents(const FString& FileName, TArray<uint8>& FileContents)
+{
+	const TArray<FCloudFile>* FilesPtr = &Files;
+	if (FilesPtr != NULL)
+	{
+		for (TArray<FCloudFile>::TConstIterator It(*FilesPtr); It; ++It)
+		{
+			if (It->FileName == FileName)
+			{
+				FileContents = It->Data;
+				return true;
+			}
+		}
+	}
+	return false;
+}
+
+bool FCdnNewsFeedTitleFile::ClearFiles()
+{
+	for (int Idx=0; Idx < Files.Num(); Idx++)
+	{
+		if (Files[Idx].AsyncState == EOnlineAsyncTaskState::InProgress)
+		{
+			UE_LOG(LogEpicStorage, Warning, TEXT("Cant clear files. Pending file op for %s"), *Files[Idx].FileName);
+			return false;
+		}
+	}
+	// remove all cached file entries
+	Files.Empty();
+	return true;
+
+}
+
+bool FCdnNewsFeedTitleFile::ClearFile(const FString& FileName)
+{
+	for (int Idx=0; Idx < Files.Num(); Idx++)
+	{
+		if (Files[Idx].FileName == FileName)
+		{
+			if (Files[Idx].AsyncState == EOnlineAsyncTaskState::InProgress)
+			{
+				UE_LOG(LogEpicStorage, Warning, TEXT("Cant clear file. Pending file op for %s"), *Files[Idx].FileName);
+				return false;
+			}
+			else
+			{
+				Files.RemoveAt(Idx);
+				return true;
+			}
+		}
+	}
+	return false;
+}
+
+bool FCdnNewsFeedTitleFile::EnumerateFiles()
+{
+	// Make sure an enumeration request  is not currently pending
+	if(!EnumerateFilesRequests.IsEmpty())
+	{
+		UE_LOG(LogEpicStorage, Warning, TEXT("EnumerateFiles request failed. Request already in progress."));
+		TriggerOnEnumerateFilesCompleteDelegates(false);
+		return false;
+	}
+
+	// Create the Http request and add to pending request list
+	TSharedRef<class IHttpRequest> HttpRequest = FHttpModule::Get().CreateRequest();
+	EnumerateFilesRequests.Enqueue(&HttpRequest.Get());
+
+	HttpRequest->OnProcessRequestComplete().BindRaw(this, &FCdnNewsFeedTitleFile::EnumerateFiles_HttpRequestComplete);
+	HttpRequest->SetURL( IndexUrl );
+	HttpRequest->SetVerb(TEXT("GET"));
+	HttpRequest->ProcessRequest();
+	return true;
+	
+}	
+
+void FCdnNewsFeedTitleFile::GetFileList(TArray<FCloudFileHeader>& OutFiles) 
+{
+	TArray<FCloudFileHeader>* FilesPtr = &FileHeaders;
+	if (FilesPtr != NULL)
+	{
+		OutFiles = *FilesPtr;
+	}
+	else
+	{
+		OutFiles.Empty();
+	}
+}
+
+FCloudFile* FCdnNewsFeedTitleFile::GetCloudFile(const FString& FileName, bool bCreateIfMissing)
+{
+	FCloudFile* CloudFile = NULL;
+	for (int Idx=0; Idx < Files.Num(); Idx++)
+	{
+		if (Files[Idx].FileName == FileName)
+		{
+			CloudFile = &Files[Idx];
+			break;
+		}
+	}
+	if (CloudFile == NULL && bCreateIfMissing)
+	{
+		CloudFile = new(Files) FCloudFile(FileName);
+	}
+	return CloudFile;
+}
+
+FString FCdnNewsFeedTitleFile::GetLocalFilePath(const FString& FileName)
+{
+	return FPaths::CloudDir() +  FPaths::GetCleanFilename(FileName);
+}
+
+void FCdnNewsFeedTitleFile::SaveCloudFileToDisk(const FString& Filename, const TArray<uint8>& Data)
+{
+	// save local disk copy as well
+	FString LocalFilePath = GetLocalFilePath( Filename );
+	bool bSavedLocal = FFileHelper::SaveArrayToFile(Data, *LocalFilePath);
+	if (bSavedLocal)
+	{
+		UE_LOG(LogEpicStorage, Verbose, TEXT("WriteUserFile request complete. Local file cache updated =%s"), 
+			*LocalFilePath);
+	}
+	else
+	{
+		UE_LOG(LogEpicStorage, Warning, TEXT("WriteUserFile request complete. Local file cache failed to update =%s"), 
+			*LocalFilePath);
+	}
+}
+
+FCloudFileHeader* FCdnNewsFeedTitleFile::GetCloudFileHeader(const FString& FileName)
+{
+	FCloudFileHeader* CloudFileHeader = NULL;
+
+	for (int Idx=0; Idx < FileHeaders.Num(); Idx++)
+	{
+		if (FileHeaders[Idx].DLName == FileName)
+		{
+			CloudFileHeader = &FileHeaders[Idx];
+			break;
+		}
+	}
+	
+	return CloudFileHeader;
+}
+
+bool FCdnNewsFeedTitleFile::ReadFile(const FString& FileName) 
+{
+	FString ErrorStr;
+
+	// Make sure valid filename was specified
+	if (FileName.IsEmpty() || FileName.Contains(TEXT(" ")))
+	{
+		ErrorStr = FString::Printf(TEXT("Invalid filename filename=%s"), *FileName);
+	}
+	
+	// Make sure a file request for this file is not currently pending
+	for (TMap<IHttpRequest*, FPendingFileRequest>::TConstIterator It(FileRequests); It; ++It)
+	{
+		if (It.Value() == FPendingFileRequest(FileName))
+		{
+			ErrorStr = FString::Printf(TEXT("File request already pending for filename=%s"), *FileName);
+			break;
+		}
+	}
+	
+	if (!ErrorStr.IsEmpty())
+	{
+		UE_LOG(LogEpicStorage, Warning, TEXT("ReadFile request failed. %s"), *ErrorStr);
+		TriggerOnReadFileCompleteDelegates(false, FileName);
+		return false;
+	}
+
+	// Mark file entry as in progress
+	FCloudFile* CloudFile = GetCloudFile(FileName, true);
+	CloudFile->AsyncState = EOnlineAsyncTaskState::InProgress;
+
+	
+	// load file from disk		
+	FString LocalFilePath = GetLocalFilePath(FileName);
+	bool bLoadedFile = FFileHelper::LoadFileToArray(CloudFile->Data, *LocalFilePath, FILEREAD_Silent);
+	if (bLoadedFile)
+	{
+		UE_LOG(LogEpicStorage, Verbose, TEXT("ReadFile request. Local file read from cache =%s"), 
+				*LocalFilePath);
+
+		//@todo samz - move reading/hashing of local file to async task
+
+		// verify hash of file if it exists
+		FCloudFileHeader* CloudFileHeader = GetCloudFileHeader( FileName);
+		if (CloudFileHeader != NULL &&
+			!CloudFileHeader->Hash.IsEmpty())
+		{
+			uint8 LocalHash[20];
+			FSHA1::HashBuffer(CloudFile->Data.GetData(), CloudFile->Data.Num(), LocalHash);
+			// concatenate 20 bye SHA1 hash to string
+			FString LocalHashStr;
+			for (int Idx=0; Idx < 20; Idx++)
+			{
+				LocalHashStr += FString::Printf(TEXT("%02x"),LocalHash[Idx]);
+			}
+			// if hash matches then just use the local file
+			if (CloudFileHeader->Hash == LocalHashStr)
+			{
+				UE_LOG(LogEpicStorage, Verbose, TEXT("Local file hash matches cloud header. No need to download for filename=%s"),  *FileName);
+
+				CloudFile->AsyncState = EOnlineAsyncTaskState::Done;
+				TriggerOnReadFileCompleteDelegates(true, FileName);
+				return true;
+			}
+		}
+	}
+	else
+	{
+		UE_LOG(LogEpicStorage, Verbose, TEXT("ReadFile request. Local file failed to read from cache =%s"), 
+				*LocalFilePath);
+	}
+	
+	// Empty local that was loaded
+	CloudFile->Data.Empty();
+
+	// Create the Http request and add to pending request list
+	TSharedRef<class IHttpRequest> HttpRequest = FHttpModule::Get().CreateRequest();
+	FileRequests.Add(&HttpRequest.Get(), FPendingFileRequest(FileName));
+
+	HttpRequest->OnProcessRequestComplete().BindRaw(this, &FCdnNewsFeedTitleFile::ReadFile_HttpRequestComplete);
+	HttpRequest->SetURL( FileName );
+	HttpRequest->SetVerb(TEXT("GET"));
+
+	return HttpRequest->ProcessRequest();
+}
+
+void FCdnNewsFeedTitleFile::EnumerateFiles_HttpRequestComplete(FHttpRequestPtr HttpRequest, FHttpResponsePtr HttpResponse, bool bSucceeded)
+{
+	IHttpRequest* Request;
+	EnumerateFilesRequests.Dequeue(Request);
+
+	bool bResult = false;
+	FString ResponseStr, ErrorStr;
+
+	if (bSucceeded && HttpResponse.IsValid())
+	{
+		ResponseStr = HttpResponse->GetContentAsString();
+		if (EHttpResponseCodes::IsOk(HttpResponse->GetResponseCode()))
+		{
+			UE_LOG(LogEpicStorage, Verbose, TEXT("EnumerateFiles request complete. url=%s code=%d response=%s"), 
+				*HttpRequest->GetURL(), HttpResponse->GetResponseCode(), *ResponseStr);
+
+			FileHeaders.Empty();
+
+			// Make sure the response is an array
+			if (!ResponseStr.StartsWith(TEXT("[")))
+			{
+				ResponseStr = FString(TEXT("[")) + ResponseStr + FString(TEXT("]"));
+			}
+
+			// Json parser expects arrays to always be wrapped with object
+			FString ResponseStrJson = FString(TEXT("{\"files\":")) + ResponseStr + FString(TEXT("}"));
+			// Create the Json parser
+			TSharedPtr<FJsonObject> JsonObject;
+			TSharedRef<TJsonReader<> > JsonReader = TJsonReaderFactory<>::Create(ResponseStrJson);
+
+			if (FJsonSerializer::Deserialize(JsonReader,JsonObject) &&
+				JsonObject.IsValid())
+			{
+				// Parse the array of file headers
+				TArray<TSharedPtr<FJsonValue> > JsonFileHeaders = JsonObject->GetArrayField(TEXT("files"));
+				for (TArray<TSharedPtr<FJsonValue> >::TConstIterator It(JsonFileHeaders); It; ++It)
+				{
+					TSharedPtr<FJsonObject> JsonFileHeader = (*It)->AsObject();
+					if (JsonFileHeader.IsValid())
+					{
+						FCloudFileHeader FileHeader;
+						if (JsonFileHeader->HasField(TEXT("hash")))
+						{
+							FileHeader.Hash = JsonFileHeader->GetStringField(TEXT("hash"));
+						}
+						if (JsonFileHeader->HasField(TEXT("uniqueFilename")))
+						{
+							FileHeader.DLName = JsonFileHeader->GetStringField(TEXT("uniqueFilename"));
+						}
+						if (JsonFileHeader->HasField(TEXT("filename")))
+						{
+							FileHeader.FileName = JsonFileHeader->GetStringField(TEXT("filename"));
+						}
+						if (JsonFileHeader->HasField(TEXT("length")))
+						{
+							FileHeader.FileSize = FMath::Trunc(JsonFileHeader->GetNumberField(TEXT("length")));
+						}
+
+						if (FileHeader.FileName.IsEmpty())
+						{
+							FileHeader.FileName = FileHeader.DLName;
+						}
+
+						if (FileHeader.Hash.IsEmpty() ||
+							FileHeader.DLName.IsEmpty())
+						{
+							UE_LOG(LogEpicStorage, Warning, TEXT("Invalid file entry hash=%s dlname=%s filename=%s"),
+								*FileHeader.Hash, *FileHeader.DLName, *FileHeader.FileName);
+						}
+						else
+						{
+							FileHeaders.Add(FileHeader);
+						}
+					}
+				}
+				bResult = true;
+			}	
+			else
+			{
+				ErrorStr = FString::Printf(TEXT("Invalid response payload=%s"),
+					*ResponseStr);
+			}
+			
+		}
+		else
+		{
+			ErrorStr = FString::Printf(TEXT("Invalid response. code=%d error=%s"),
+				HttpResponse->GetResponseCode(), *ResponseStr);
+		}
+	}
+
+
+	if (!ErrorStr.IsEmpty())
+	{
+		UE_LOG(LogEpicStorage, Warning, TEXT("EnumerateFiles request failed. %s"), *ErrorStr);
+	}
+
+	TriggerOnEnumerateFilesCompleteDelegates(bResult);
+}
+
+void FCdnNewsFeedTitleFile::ReadFile_HttpRequestComplete(FHttpRequestPtr HttpRequest, FHttpResponsePtr HttpResponse, bool bSucceeded)
+{
+	bool bResult = false;
+	FString ResponseStr, ErrorStr;
+
+	// should have a pending Http request
+	FPendingFileRequest PendingRequest = FileRequests.FindChecked(HttpRequest.Get());
+	FileRequests.Remove(HttpRequest.Get());
+
+	// Cloud file being operated on
+	FCloudFile* CloudFile = GetCloudFile(PendingRequest.FileName, true);
+	CloudFile->AsyncState = EOnlineAsyncTaskState::Failed;
+	CloudFile->Data.Empty();
+	
+	if (bSucceeded && 
+		HttpResponse.IsValid())
+	{
+		
+		if (EHttpResponseCodes::IsOk(HttpResponse->GetResponseCode()))
+		{
+
+			UE_LOG(LogEpicStorage, Verbose, TEXT("ReadFile request complete. url=%s code=%d"), 
+				*HttpRequest->GetURL(), HttpResponse->GetResponseCode());
+
+			// update the memory copy of the file with data that was just downloaded
+			FCloudFile* CloudFile = GetCloudFile(PendingRequest.FileName, true);
+			CloudFile->AsyncState = EOnlineAsyncTaskState::Done;
+			CloudFile->Data = HttpResponse->GetContent();
+
+			// cache to disk on successful download
+			SaveCloudFileToDisk(CloudFile->FileName,CloudFile->Data);
+
+			bResult = true;
+		}
+		else
+		{
+			ErrorStr = FString::Printf(TEXT("Invalid response. code=%d error=%s"),
+				HttpResponse->GetResponseCode(), *HttpResponse->GetContentAsString());
+		}
+	}
+	else
+	{
+		ErrorStr = TEXT("No response");
+	}
+	if (!ErrorStr.IsEmpty())
+	{
+		UE_LOG(LogEpicStorage, Warning, TEXT("ReadFile request failed. %s"), *ErrorStr);
+	}
+
+	TriggerOnReadFileCompleteDelegates(bResult, PendingRequest.FileName);
+}
\ No newline at end of file
diff --git a/Engine/Source/Editor/NewsFeed/Private/Implementation/CdnNewsFeedTitleFile.h b/Engine/Source/Editor/NewsFeed/Private/Implementation/CdnNewsFeedTitleFile.h
new file mode 100644
--- /dev/null
+++ b/Engine/Source/Editor/NewsFeed/Private/Implementation/CdnNewsFeedTitleFile.h
@@ -0,0 +1,131 @@
+// Copyright 1998-2014 Epic Games, Inc. All Rights Reserved.
+#pragma once
+
+
+/**
+ * CDN based title file implementation.
+ * This code was salvaged from the deprecated EpicContent plug-in.
+ */
+class FCdnNewsFeedTitleFile
+	: public IOnlineTitleFile
+{
+	/** Info used to send request for a file */
+	struct FPendingFileRequest
+	{
+		/**
+		 * Constructor
+		 */
+		FPendingFileRequest(const FString& InFileName=FString(TEXT("")))
+			:  FileName(InFileName)
+		{
+
+		}
+
+		/**
+		 * Equality op
+		 */
+		inline bool operator==(const FPendingFileRequest& Other) const
+		{
+			return FileName == Other.FileName;
+		}
+
+		/** File being operated on by the pending request */
+		FString FileName;
+	};
+
+public:
+
+	static IOnlineTitleFilePtr Create( const FString& IndexUrl );
+
+public:
+
+	// Begin IOnlineTitleFile interface
+
+	virtual bool GetFileContents(const FString& DLName, TArray<uint8>& FileContents) OVERRIDE;
+
+	virtual bool ClearFiles() OVERRIDE;
+
+	virtual bool ClearFile(const FString& DLName) OVERRIDE;
+
+	virtual bool EnumerateFiles() OVERRIDE;
+
+	virtual void GetFileList(TArray<FCloudFileHeader>& InFileHeaders) OVERRIDE;
+
+	virtual bool ReadFile(const FString& DLName) OVERRIDE;
+
+	// Begin IOnlineTitleFile interface
+
+protected:
+	
+	/**
+	 * Delegate called when a Http request completes for enumerating list of file headers
+	 */
+	void EnumerateFiles_HttpRequestComplete(FHttpRequestPtr HttpRequest, FHttpResponsePtr HttpResponse, bool bSucceeded);
+
+	/**
+	 * Delegate called when a Http request completes for reading a cloud file
+	 */
+	void ReadFile_HttpRequestComplete(FHttpRequestPtr HttpRequest, FHttpResponsePtr HttpResponse, bool bSucceeded);
+
+	/**
+	 * Find/create cloud file entry 
+	 *
+	 * @param FileName cached file entry to find
+	 * @param bCreateIfMissing create the file entry if not found
+	 *
+	 * @return cached cloud file entry
+	 */
+	FCloudFile* GetCloudFile(const FString& FileName, bool bCreateIfMissing);
+
+	/**
+	 * Find cloud file header entry 
+	 *
+	 * @param FileName cached file entry to find
+	 *
+	 * @return cached cloud file header entry
+	 */
+	FCloudFileHeader* GetCloudFileHeader( const FString& FileName);
+
+	/**
+	 * Converts filename into a local file cache path
+	 *
+	 * @param FileName name of file being loaded
+	 *
+	 * @return unreal file path to be used by file manager
+	 */
+	FString GetLocalFilePath(const FString& FileName);
+
+	/**
+	 * Save a file from a given user to disk
+	 *
+	 * @param FileName name of file being saved
+	 * @param Data data to write to disk
+	 */
+	void SaveCloudFileToDisk(const FString& Filename, const TArray<uint8>& Data);
+
+private:
+
+	/**
+	 * Should use the initialization constructor instead
+	 */
+	FCdnNewsFeedTitleFile( const FString& InIndexUrl );
+
+private:
+
+	// Config based url for enumerating list of cloud files.
+	FString EnumerateFilesUrl;
+	
+	// Config based url for downloading/updating a single cloud file entry.
+	FString FileUrl;
+
+	// List of pending Http requests for enumerating files.
+	TQueue<IHttpRequest*> EnumerateFilesRequests;
+
+	// List of pending Http requests for reading files.
+	TMap<class IHttpRequest*, FPendingFileRequest> FileRequests;
+
+	TArray<FCloudFileHeader> FileHeaders;
+	TArray<FCloudFile> Files;
+
+	FString IndexUrl;
+};
diff --git a/Engine/Source/Editor/NewsFeed/Private/Implementation/LocalNewsFeedTitleFile.cpp b/Engine/Source/Editor/NewsFeed/Private/Implementation/LocalNewsFeedTitleFile.cpp
new file mode 100644
--- /dev/null
+++ b/Engine/Source/Editor/NewsFeed/Private/Implementation/LocalNewsFeedTitleFile.cpp
@@ -0,0 +1,138 @@
+// Copyright 1998-2014 Epic Games, Inc. All Rights Reserved.
+
+#include "NewsFeedPrivatePCH.h"
+
+
+IOnlineTitleFilePtr FLocalNewsFeedTitleFile::Create( const FString& RootDirectory )
+{
+	return MakeShareable( new FLocalNewsFeedTitleFile( RootDirectory ) );
+}
+
+
+FLocalNewsFeedTitleFile::FLocalNewsFeedTitleFile( const FString& InRootDirectory )
+	: RootDirectory(InRootDirectory)
+{ }
+
+
+bool FLocalNewsFeedTitleFile::GetFileContents(const FString& DLName, TArray<uint8>& FileContents)
+{
+	const TArray< uint8 >* ExistingFileContents = DLNameToFileContents.Find( DLName );
+
+	if (ExistingFileContents == nullptr)
+	{
+		return false;
+	}
+
+	FileContents.Append(*ExistingFileContents);
+
+	return true;
+}
+
+
+bool FLocalNewsFeedTitleFile::ClearFiles( )
+{
+	FileHeaders.Empty();
+	DLNameToFileContents.Empty();
+
+	return true;
+}
+
+
+bool FLocalNewsFeedTitleFile::ClearFile( const FString& DLName )
+{
+	bool ClearedFile = false;
+	const FString FileName = GetFileNameFromDLName(DLName);
+
+	for (int Index = 0; Index < FileHeaders.Num(); Index++)
+	{
+		if (FileHeaders[Index].DLName == DLName)
+		{
+			FileHeaders.RemoveAt(Index);
+			ClearedFile = true;
+		}
+	}
+
+	DLNameToFileContents.Remove(DLName);
+
+	return ClearedFile;
+}
+
+
+bool FLocalNewsFeedTitleFile::EnumerateFiles()
+{
+	const FString WildCard = FPaths::Combine(*RootDirectory, TEXT("*"));
+
+	TArray<FString> Filenames;
+	IFileManager::Get().FindFiles(Filenames, *WildCard, true, false);
+
+	bool Success = Filenames.Num() > 0;
+
+	if (Success)
+	{
+		for (int32 FileIdx = 0; FileIdx < Filenames.Num(); ++FileIdx)
+		{
+			const FString Filename = Filenames[FileIdx];
+
+			FCloudFileHeader NewHeader;
+			NewHeader.FileName = Filename;
+			NewHeader.DLName = Filename + FileIdx;
+			NewHeader.FileSize = 0;
+			NewHeader.Hash.Empty();
+
+			FileHeaders.Add( NewHeader );
+		}
+	}
+
+	TriggerOnEnumerateFilesCompleteDelegates(Success);
+
+	return Success;
+}
+
+
+void FLocalNewsFeedTitleFile::GetFileList(TArray<FCloudFileHeader>& InFileHeaders)
+{
+	InFileHeaders.Append( FileHeaders );
+}
+
+
+bool FLocalNewsFeedTitleFile::ReadFile(const FString& DLName)
+{
+	const TArray< uint8 >* ExistingFileContents = DLNameToFileContents.Find( DLName );
+
+	if (ExistingFileContents != nullptr)
+	{
+		TriggerOnReadFileCompleteDelegates(true, DLName);
+
+		return true;
+	}
+
+	const FString FileName = GetFileNameFromDLName( DLName );
+
+	TArray<uint8> FileContents;
+
+	if (!FFileHelper::LoadFileToArray(FileContents, *(RootDirectory + FileName)))
+	{
+		TriggerOnReadFileCompleteDelegates(false, DLName);
+
+		return false;
+	}
+
+	DLNameToFileContents.Add(DLName, FileContents);
+	TriggerOnReadFileCompleteDelegates(true, DLName);
+
+	return true;
+}
+
+
+FString FLocalNewsFeedTitleFile::GetFileNameFromDLName( const FString& DLName ) const
+{
+	for (int Index = 0; Index < FileHeaders.Num(); Index++)
+	{
+		if (FileHeaders[Index].DLName == DLName)
+		{
+			return FileHeaders[Index].FileName;
+		}
+	}
+
+	return FString();
+}
diff --git a/Engine/Source/Editor/NewsFeed/Private/Implementation/LocalNewsFeedTitleFile.h b/Engine/Source/Editor/NewsFeed/Private/Implementation/LocalNewsFeedTitleFile.h
new file mode 100644
--- /dev/null
+++ b/Engine/Source/Editor/NewsFeed/Private/Implementation/LocalNewsFeedTitleFile.h
@@ -0,0 +1,47 @@
+// Copyright 1998-2014 Epic Games, Inc. All Rights Reserved.
+#pragma once
+
+
+/**
+ * Local file system based title file implementation (for testing).
+ * This code was salvaged from the deprecated EpicContent plug-in.
+ */
+class FLocalNewsFeedTitleFile
+	: public IOnlineTitleFile
+{
+public:
+
+	static IOnlineTitleFilePtr Create( const FString& RootDirectory );
+
+public:
+
+	// Begin IOnlineTitleFile interface
+
+	virtual bool GetFileContents(const FString& DLName, TArray<uint8>& FileContents) OVERRIDE;
+
+	virtual bool ClearFiles() OVERRIDE;
+
+	virtual bool ClearFile(const FString& DLName) OVERRIDE;
+
+	virtual bool EnumerateFiles() OVERRIDE;
+
+	virtual void GetFileList(TArray<FCloudFileHeader>& InFileHeaders) OVERRIDE;
+
+	virtual bool ReadFile(const FString& DLName) OVERRIDE;
+
+	// End IOnlineTitleFile interface
+
+protected:
+
+	FString GetFileNameFromDLName( const FString& DLName ) const;
+
+private:
+
+	FLocalNewsFeedTitleFile( const FString& InRootDirectory );
+
+private:
+
+	FString RootDirectory;
+	TArray<FCloudFileHeader> FileHeaders;
+	TMap< FString, TArray<uint8> > DLNameToFileContents;
+};
diff --git a/Engine/Source/Editor/NewsFeed/Private/Implementation/NewsFeedCache.cpp b/Engine/Source/Editor/NewsFeed/Private/Implementation/NewsFeedCache.cpp
new file mode 100644
--- /dev/null
+++ b/Engine/Source/Editor/NewsFeed/Private/Implementation/NewsFeedCache.cpp
@@ -0,0 +1,388 @@
+// Copyright 1998-2014 Epic Games, Inc. All Rights Reserved.
+
+#include "NewsFeedPrivatePCH.h"
+
+
+/* Constants
+ *****************************************************************************/
+
+const FString FeedFilename(TEXT("newsfeed.json"));
+
+
+/* FNewsFeedCache structors
+ *****************************************************************************/
+
+FNewsFeedCache::FNewsFeedCache( )
+	: CurrentCultureSpec(ENewsFeedCultureSpec::Full)
+	, LoaderState(ENewsFeedState::NotStarted)
+{
+	TickDelegate = FTickerDelegate::CreateRaw(this, &FNewsFeedCache::HandleTicker);
+}
+
+
+FNewsFeedCache::~FNewsFeedCache( )
+{
+	FTicker::GetCoreTicker().RemoveTicker(TickDelegate);
+	SaveSettings();
+}
+
+
+/* FNewsFeedCache interface
+ *****************************************************************************/
+
+const FSlateBrush* FNewsFeedCache::GetIcon( const FString& IconName ) const
+{
+	const TSharedPtr<FSlateDynamicImageBrush>* IconBrush = Icons.Find(IconName);
+
+	if (IconBrush != nullptr)
+	{
+		return IconBrush->Get();
+	}
+
+	return FEditorStyle::GetBrush(TEXT("NewsFeed.PendingIcon"));
+}
+
+
+//bool FNewsFeedCache::LoadFeed( IOnlineTitleFilePtr InTitleFile )
+bool FNewsFeedCache::LoadFeed(  )
+{
+	if (LoaderState == ENewsFeedState::Loading)
+	{
+		return false;
+	}
+
+//	if (!InTitleFile.IsValid())
+//	{
+//		return false;
+//	}
+
+	LoaderState = ENewsFeedState::Loading;
+	Icons.Empty();
+
+//	TitleFile = InTitleFile;
+//	TitleFile->AddOnEnumerateFilesCompleteDelegate(FOnEnumerateFilesCompleteDelegate::CreateRaw(this, &FNewsFeedCache::HandleEnumerateFilesComplete));
+//	TitleFile->AddOnReadFileCompleteDelegate(FOnReadFileCompleteDelegate::CreateRaw(this, &FNewsFeedCache::HandleReadFileComplete));
+//	TitleFile->ClearFiles();
+//	TitleFile->EnumerateFiles();
+	LoadTitleFile();
+
+	// auto reload after 15 minutes
+	FTicker::GetCoreTicker().RemoveTicker(TickDelegate);
+	FTicker::GetCoreTicker().AddTicker(TickDelegate, 15 * 60.0f);
+
+	return true;
+}
+
+
+/* FNewsFeedCache implementation
+ *****************************************************************************/
+
+void FNewsFeedCache::FailLoad( )
+{
+	LoaderState = ENewsFeedState::Failure;
+	LoadFailedDelegate.Broadcast();
+}
+
+
+void FNewsFeedCache::LoadTitleFile( )
+{
+	FString CultureString;
+
+	// get current culture
+	TSharedRef<FCulture> Culture = FInternationalization::Get().GetCurrentCulture();
+
+	if (CurrentCultureSpec == ENewsFeedCultureSpec::Full)
+	{
+		CultureString = Culture->GetName();
+	}
+	else if (CurrentCultureSpec == ENewsFeedCultureSpec::LanguageOnly)
+	{
+		CultureString = Culture->GetTwoLetterISOLanguageName();
+	}
+
+	// get title file source
+	FString SourceType;
+
+	if (GConfig->GetString(TEXT("NewsFeed"), TEXT("Source"), SourceType, GEditorIni) && !SourceType.IsEmpty())
+	{
+		if (SourceType == TEXT("CDN"))
+		{
+			FString CdnUrl;
+
+			if (GConfig->GetString(TEXT("NewsFeed"),TEXT("CdnUrl"), CdnUrl, GEditorIni) && !CdnUrl.IsEmpty())
+			{
+				if (!CultureString.IsEmpty())
+				{
+					CdnUrl = CdnUrl.Replace( TEXT("{Culture}"), *CultureString );
+				}
+				else
+				{
+					CdnUrl = CdnUrl.Replace(TEXT("{Culture}/"), TEXT(""));
+				}
+
+				TitleFile = FCdnNewsFeedTitleFile::Create(CdnUrl);
+			}
+		}
+		else if (SourceType == TEXT("LOCAL"))
+		{
+			FString EngineContentRelativeDirectory;
+
+			if (GConfig->GetString(TEXT("NewsFeed"), TEXT("EngineContentRelativeDirectory"), EngineContentRelativeDirectory, GEditorIni) && !EngineContentRelativeDirectory.IsEmpty())
+			{
+				FString RootDirectory = FPaths::Combine(FPlatformProcess::BaseDir(), *FPaths::EngineContentDir(), *EngineContentRelativeDirectory);
+
+				if (!CultureString.IsEmpty())
+				{
+					RootDirectory = RootDirectory.Replace(TEXT("{Culture}"), *CultureString);
+				}
+				else
+				{
+					RootDirectory = RootDirectory.Replace(TEXT("{Culture}/"), TEXT(""));
+				}
+
+				TitleFile = FLocalNewsFeedTitleFile::Create(RootDirectory);
+			}
+		}
+		else if (SourceType == TEXT("MCP"))
+		{
+			TitleFile = Online::GetTitleFileInterface(TEXT("MCP"));
+		}
+	}
+
+	if (TitleFile.IsValid())
+	{
+		TitleFile->AddOnEnumerateFilesCompleteDelegate(FOnEnumerateFilesCompleteDelegate::CreateRaw(this, &FNewsFeedCache::HandleEnumerateFilesComplete));
+		TitleFile->AddOnReadFileCompleteDelegate(FOnReadFileCompleteDelegate::CreateRaw(this, &FNewsFeedCache::HandleReadFileComplete));
+		TitleFile->ClearFiles();
+		TitleFile->EnumerateFiles();
+	}
+	else
+	{
+		FailLoad();
+	}
+}
+
+
+void FNewsFeedCache::ProcessIconFile( const FString& IconName )
+{
+	TArray<uint8> FileContents;
+	TitleFile->GetFileContents(IconName, FileContents);
+
+	TSharedPtr<FSlateDynamicImageBrush> IconBrush = RawDataToBrush(*(FString(TEXT("NewsFeed.")) + IconName), FileContents);
+	Icons.Add(IconName, IconBrush);
+}
+
+
+void FNewsFeedCache::ProcessNewsFeedFile( )
+{
+	Items.Empty();
+
+	// read news feed file
+	TArray<uint8> FileContents;
+	TitleFile->GetFileContents(FeedFileHeader.DLName, FileContents);
+
+	FString NewsFeedJson;
+	FFileHelper::BufferToString(NewsFeedJson, FileContents.GetData(), FileContents.Num());
+
+	// deserialize the news feed Json
+	TSharedRef<TJsonReader<> > Reader = TJsonReaderFactory<TCHAR>::Create(NewsFeedJson);
+	TSharedPtr<FJsonObject> FeedObject;
+
+	if (!FJsonSerializer::Deserialize(Reader, FeedObject))
+	{
+		FailLoad();
+
+		return;
+	}
+
+	// get the file list again for any icons we may have to load
+	TArray<FCloudFileHeader> FileHeaders;
+	TitleFile->GetFileList(FileHeaders);
+
+	// parse the news feed items
+	const UNewsFeedSettings* NewsFeedSettings = GetDefault<UNewsFeedSettings>();
+	TArray<TSharedPtr<FJsonValue> > ItemObjects = FeedObject->GetArrayField(TEXT("news"));
+
+	// for each news item...
+	for (int32 ItemObjectIndex = 0; ItemObjectIndex < ItemObjects.Num(); ++ItemObjectIndex)
+	{
+		const TSharedPtr<FJsonObject>& ItemObject = ItemObjects[ItemObjectIndex]->AsObject();
+
+		// ... that has a valid identifier...
+		if (ItemObject->HasTypedField<EJson::String>("id"))
+		{
+			FGuid ItemId;
+
+			if (FGuid::Parse(ItemObject->GetStringField("id"), ItemId))
+			{
+				// .. parse its data...
+				FNewsFeedItemRef NewsFeedItem = MakeShareable(new FNewsFeedItem(ItemId));
+
+				NewsFeedItem->Excerpt = FText::FromString(ItemObject->GetStringField("excerpt"));
+				NewsFeedItem->FullText = FText::FromString(ItemObject->GetStringField("fulltext"));
+				NewsFeedItem->Read = NewsFeedSettings->ReadItems.Contains(ItemId);
+				NewsFeedItem->Title = FText::FromString(ItemObject->GetStringField("title"));
+				NewsFeedItem->Url = ItemObject->GetStringField("url");
+
+				FDateTime Issued;
+
+				if (FDateTime::Parse(ItemObject->GetStringField("issued"), Issued))
+				{
+					NewsFeedItem->Issued = Issued;
+				}
+				else
+				{
+					NewsFeedItem->Issued = FDateTime::MinValue();
+				}
+
+				// ... and load its icon image
+				const FString IconName = ItemObject->GetStringField("icon");
+
+				for (int32 FileHeaderIndex = 0; FileHeaderIndex < FileHeaders.Num(); ++FileHeaderIndex)
+				{
+					const FCloudFileHeader& FileHeader = FileHeaders[FileHeaderIndex];
+
+					if (FileHeader.FileName == IconName)
+					{
+						NewsFeedItem->IconName = FileHeader.DLName;
+						TitleFile->ReadFile(FileHeader.DLName);
+
+						break;
+					}
+				}
+
+				Items.Add(NewsFeedItem);
+			}
+		}
+	}
+
+	LoaderState = ENewsFeedState::Success;
+	LoadCompletedDelegate.Broadcast();
+}
+
+
+TSharedPtr<FSlateDynamicImageBrush> FNewsFeedCache::RawDataToBrush( FName ResourceName, const TArray< uint8 >& InRawData ) const
+{
+	TSharedPtr<FSlateDynamicImageBrush> Brush;
+
+	uint32 BytesPerPixel = 4;
+	int32 Width = 0;
+	int32 Height = 0;
+
+	bool bSucceeded = false;
+	TArray<uint8> DecodedImage;
+	IImageWrapperModule& ImageWrapperModule = FModuleManager::LoadModuleChecked<IImageWrapperModule>(FName("ImageWrapper"));
+	IImageWrapperPtr ImageWrapper = ImageWrapperModule.CreateImageWrapper( EImageFormat::PNG );
+
+	if (ImageWrapper.IsValid() && ImageWrapper->SetCompressed(InRawData.GetData(), InRawData.Num()))
+	{
+		Width = ImageWrapper->GetWidth();
+		Height = ImageWrapper->GetHeight();
+
+		const TArray<uint8>* RawData = NULL;
+
+		if (ImageWrapper->GetRaw( ERGBFormat::BGRA, 8, RawData))
+		{
+			DecodedImage.AddUninitialized( Width * Height * BytesPerPixel );
+			DecodedImage = *RawData;
+			bSucceeded = true;
+		}
+	}
+
+	if (bSucceeded && FSlateApplication::Get().GetRenderer()->GenerateDynamicImageResource(ResourceName, ImageWrapper->GetWidth(), ImageWrapper->GetHeight(), DecodedImage))
+	{
+		Brush = MakeShareable(new FSlateDynamicImageBrush( ResourceName, FVector2D(ImageWrapper->GetWidth(), ImageWrapper->GetHeight()) ) );
+	}
+
+	return Brush;
+}
+
+
+void FNewsFeedCache::SaveSettings( )
+{
+	UNewsFeedSettings* NewsFeedSettings = GetMutableDefault<UNewsFeedSettings>();
+
+	// remember read items
+	for (auto NewsFeedItem : Items)
+	{
+		if (NewsFeedItem->Read)
+		{
+			NewsFeedSettings->ReadItems.AddUnique(NewsFeedItem->ItemId);
+		}
+	}
+
+	NewsFeedSettings->SaveConfig();
+}
+
+
+/* FNewsFeedCache callbacks
+ *****************************************************************************/
+
+void FNewsFeedCache::HandleEnumerateFilesComplete( bool bSuccess )
+{
+	if (bSuccess)
+	{
+		// load news feed files
+		TArray<FCloudFileHeader> FileHeaders;
+		TitleFile->GetFileList(FileHeaders);
+
+		for (int32 FileHeaderIndex = 0; FileHeaderIndex < FileHeaders.Num(); ++FileHeaderIndex)
+		{
+			const FCloudFileHeader& FileHeader = FileHeaders[FileHeaderIndex];
+
+			if (FileHeader.FileName == FeedFilename)
+			{
+				FeedFileHeader = FileHeader;
+				TitleFile->ReadFile(FileHeader.DLName);
+
+				return;
+			}
+		}
+
+		FailLoad();
+	}
+	else
+	{
+		// move on to the next culture specification and try again
+		if (CurrentCultureSpec != ENewsFeedCultureSpec::None)
+		{
+			CurrentCultureSpec = static_cast<ENewsFeedCultureSpec>(int(CurrentCultureSpec) + 1);
+			LoadTitleFile();
+		}
+		else
+		{
+			FailLoad();
+		}
+	}
+}
+
+
+void FNewsFeedCache::HandleReadFileComplete( bool bSuccess, const FString& Filename )
+{
+	if (Filename == FeedFileHeader.DLName)
+	{
+		if (bSuccess)
+		{
+			ProcessNewsFeedFile();
+		}
+		else
+		{
+			FailLoad();
+		}
+	}
+	else if (bSuccess)
+	{
+		if (!Icons.Contains(Filename))
+		{
+			ProcessIconFile(Filename);
+		}		
+	}
+}
+
+
+bool FNewsFeedCache::HandleTicker( float DeltaTime )
+{
+	LoadFeed();
+
+	return false;
+}
diff --git a/Engine/Source/Editor/NewsFeed/Private/Implementation/NewsFeedCache.h b/Engine/Source/Editor/NewsFeed/Private/Implementation/NewsFeedCache.h
new file mode 100644
--- /dev/null
+++ b/Engine/Source/Editor/NewsFeed/Private/Implementation/NewsFeedCache.h
@@ -0,0 +1,211 @@
+#pragma once
+
+
+enum class ENewsFeedCultureSpec
+{
+	Full = 0,
+	LanguageOnly = 1,
+	None = 2
+};
+
+
+namespace ENewsFeedState
+{
+	enum Type
+	{
+		NotStarted = 0,
+		Loading = 1,
+		Success = 2,
+		Failure = 3
+	};
+}
+
+
+/**
+ * Type definition for shared pointers to instances of FNewsFeedCache.
+ */
+typedef TSharedPtr<class FNewsFeedCache> FNewsFeedCachePtr;
+
+/**
+ * Type definition for shared references to instances of FNewsFeedCache.
+ */
+typedef TSharedRef<class FNewsFeedCache> FNewsFeedCacheRef;
+
+
+class FNewsFeedCache
+{
+public:
+
+	/**
+	 * Default constructor.
+	 */
+	FNewsFeedCache( );
+
+	/**
+	 * Destructor.
+	 */
+	~FNewsFeedCache( );
+
+public:
+
+	/**
+	 * Gets the specified icon.
+	 *
+	 * If the icon is not yet loaded, a placeholder icon will be returned.
+	 * If the icon failed to load, an error icon will be returned.
+	 *
+	 * @param IconName The name of the icon to get.
+	 *
+	 * @return The icon brush.
+	 */
+	const FSlateBrush* GetIcon( const FString& IconName ) const;
+
+	/**
+	 * Gets the loaded news feed items.
+	 *
+	 * @return Collection of news feed items.
+	 */
+	const TArray<FNewsFeedItemPtr>& GetItems( ) const
+	{
+		return Items;
+	}
+
+	/**
+	 * Checks whether the last attempt to load a feed has finished.
+	 *
+	 * @return true if feed loading finished, false otherwise.
+	 *
+	 * @see IsLoading
+	 */
+	bool IsFinished( ) const
+	{
+		return (LoaderState == ENewsFeedState::Success);
+	}
+
+	/**
+	 * Checks whether this loader is currently loading a news feed.
+	 *
+	 * @return true if loading, false otherwise.
+	 *
+	 * @see IsFinished
+	 */
+	bool IsLoading( ) const
+	{
+		return (LoaderState == ENewsFeedState::Loading);
+	}
+
+	/**
+	 * Loads the news feed from the given title file.
+	 *
+	 * @param InTitleFile The title file to read the news feed from.
+	 *
+	 * @return true if the news feed is being loaded, false otherwise.
+	 */
+	//bool LoadFeed( IOnlineTitleFilePtr InTitleFile );
+	bool LoadFeed( );
+
+public:
+
+	/**
+	 * Gets an event delegate that is executed when loading of a news feed has completed.
+	 *
+	 * @return The event delegate.
+	 */
+	DECLARE_EVENT(FNewsFeedCache, FOnNewsFeedLoadCompleted);
+	FOnNewsFeedLoadCompleted& OnLoadCompleted( )
+	{
+		return LoadCompletedDelegate;
+	}
+
+	/**
+	 * Gets an event delegate that is executed when loading of a news feed has failed.
+	 *
+	 * @return The event delegate.
+	 */
+	DECLARE_EVENT(FNewsFeedCache, FOnNewsFeedLoadFailed);
+	FOnNewsFeedLoadFailed& OnLoadFailed( )
+	{
+		return LoadFailedDelegate;
+	}
+
+protected:
+
+	/**
+	 * Signals that the loading has failed.
+	 */
+	void FailLoad( );
+
+	/**
+	 * Attempts to load the title file.
+	 */
+	void LoadTitleFile( );
+
+	/**
+	 * Processes the specified icon file.
+	 *
+	 * @param IconName The name of the icon file.
+	 */
+	void ProcessIconFile( const FString& IconName );
+
+	/**
+	 * Processes the Json file containing the news feed items.
+	 */
+	void ProcessNewsFeedFile( );
+
+	/**
+	 * Converts the given binary data into a Slate brush.
+	 *
+	 * @param ResourceName 
+	 * @param RawData The raw binary data to convert.
+	 *
+	 * @return The Slate brush.
+	 */
+	TSharedPtr<FSlateDynamicImageBrush> RawDataToBrush( FName ResourceName, const TArray< uint8 >& RawData ) const;
+
+	/**
+	 * Saves all settings.
+	 */
+	void SaveSettings( );
+
+private:
+
+	// Callback for when file enumeration has completed.
+	void HandleEnumerateFilesComplete( bool bSuccess );
+
+	// Callback for when file reading has completed.
+	void HandleReadFileComplete( bool bSuccess, const FString& Filename );
+
+	// Callback for ticks from the ticker.
+	bool HandleTicker( float DeltaTime );
+
+private:
+
+	// Holds the header of the news feed file.
+	FCloudFileHeader FeedFileHeader;
+
+	// Holds the collection of loaded icons.
+	TMap<FString, TSharedPtr<FSlateDynamicImageBrush> > Icons;
+
+	// Holds the news feed items.
+	TArray<FNewsFeedItemPtr> Items;
+
+	// Holds the current state.
+	ENewsFeedState::Type LoaderState;
+
+	// Holds a delegate to be invoked when the feed needs to be auto-reloaded.
+	FTickerDelegate TickDelegate;
+
+	// Holds the title file to load from.
+	IOnlineTitleFilePtr TitleFile;
+
+private:
+
+	// Holds the culture currently being loaded with.
+	ENewsFeedCultureSpec CurrentCultureSpec;
+
+	// Holds an event delegate that is executed when loading of a news feed has completed.
+	FOnNewsFeedLoadCompleted LoadCompletedDelegate;
+
+	// Holds an event delegate that is executed when loading of a news feed has failed.
+	FOnNewsFeedLoadFailed LoadFailedDelegate;
+};
diff --git a/Engine/Source/Editor/NewsFeed/Private/Implementation/NewsFeedSettings.h b/Engine/Source/Editor/NewsFeed/Private/Implementation/NewsFeedSettings.h
new file mode 100644
--- /dev/null
+++ b/Engine/Source/Editor/NewsFeed/Private/Implementation/NewsFeedSettings.h
@@ -0,0 +1,60 @@
+// Copyright 1998-2014 Epic Games, Inc. All Rights Reserved.
+
+/*=============================================================================
+	NewsFeedSettings.h: Declares the UNewsFeedSettings class.
+=============================================================================*/
+
+#pragma once
+
+#include "NewsFeedSettings.generated.h"
+
+
+UCLASS(config=Editor)
+class UNewsFeedSettings
+	: public UObject
+{
+	GENERATED_UCLASS_BODY()
+
+public:
+
+	/**
+	 * The maximum number of news items to show.
+	 */
+	UPROPERTY(config, EditAnywhere, Category=Display)
+	int32 MaxItemsToShow;
+
+	/**
+	 * List of news feed items that have been marked as read.
+	 */
+	UPROPERTY(config)
+	TArray<FGuid> ReadItems;
+
+	/**
+	 * Whether to show only unread news feed items.
+	 */
+	UPROPERTY(config, EditAnywhere, Category=Display)
+	bool ShowOnlyUnreadItems;
+
+public:
+
+	/**
+	 * Returns an event delegate that is executed when a setting has changed.
+	 *
+	 * @return The delegate.
+	 */
+	DECLARE_EVENT_OneParam(UNewsFeedSettings, FSettingChangedEvent, FName /*PropertyName*/);
+	FSettingChangedEvent& OnSettingChanged( ) { return SettingChangedEvent; }
+
+protected:
+
+	// Begin UObject overrides
+
+	virtual void PostEditChangeProperty( struct FPropertyChangedEvent& PropertyChangedEvent ) OVERRIDE;
+
+	// End UObject overrides
+
+private:
+
+	// Holds an event delegate that is executed when a setting has changed.
+	FSettingChangedEvent SettingChangedEvent;
+};
diff --git a/Engine/Source/Editor/NewsFeed/Private/Models/NewsFeedItem.h b/Engine/Source/Editor/NewsFeed/Private/Models/NewsFeedItem.h
new file mode 100644
--- /dev/null
+++ b/Engine/Source/Editor/NewsFeed/Private/Models/NewsFeedItem.h
@@ -0,0 +1,74 @@
+// Copyright 1998-2014 Epic Games, Inc. All Rights Reserved.
+
+/*=============================================================================
+	NewsFeedItem.h: Declares the FNewsFeedItem class.
+=============================================================================*/
+
+#pragma once
+
+
+/**
+ * Type definition for shared pointers to instances of FNewsFeedItem.
+ */
+typedef TSharedPtr<struct FNewsFeedItem> FNewsFeedItemPtr;
+
+/**
+ * Type definition for shared references to instances of FNewsFeedItem.
+ */
+typedef TSharedRef<struct FNewsFeedItem> FNewsFeedItemRef;
+
+
+/**
+ * Implements a view model for the news feed list.
+ */
+struct FNewsFeedItem
+{
+	/**
+	 * The news excerpt.
+	 */
+	FText Excerpt;
+
+	/**
+	 * The full news text.
+	 */
+	FText FullText;
+
+	/**
+	 * Holds the name of the icon.
+	 */
+	FString IconName;
+
+	/**
+	 * The date and time at which the news item was issued.
+	 */
+	FDateTime Issued;
+
+	/**
+	 * The news item's unique identifier.
+	 */
+	FGuid ItemId;
+
+	/**
+	 * Whether this news item has been marked as read.
+	 */
+	bool Read;
+
+	/**
+	 * The news title.
+	 */
+	FText Title;
+
+	/**
+	 * The URL to the full news release.
+	 */
+	FString Url;
+
+	/**
+	 * Creates and initializes a new item with the specified identifier.
+	 *
+	 * @param InItemId - The item identifier.
+	 */
+	FNewsFeedItem( const FGuid& InItemId )
+		: ItemId(InItemId)
+	{ }
+};
diff --git a/Engine/Source/Editor/NewsFeed/Private/NewsFeedClasses.cpp b/Engine/Source/Editor/NewsFeed/Private/NewsFeedClasses.cpp
new file mode 100644
--- /dev/null
+++ b/Engine/Source/Editor/NewsFeed/Private/NewsFeedClasses.cpp
@@ -0,0 +1,27 @@
+// Copyright 1998-2014 Epic Games, Inc. All Rights Reserved.
+
+/*=============================================================================
+	NewsFeedClasses.cpp: Implements the module's script classes.
+=============================================================================*/
+
+#include "NewsFeedPrivatePCH.h"
+#include "NewsFeed.generated.inl"
+
+
+/* UNewsFeedSettings interface
+ *****************************************************************************/
+
+UNewsFeedSettings::UNewsFeedSettings( const class FPostConstructInitializeProperties& PCIP )
+	: Super(PCIP)
+	, MaxItemsToShow(10)
+	, ShowOnlyUnreadItems(true)
+{ }
+
+
+void UNewsFeedSettings::PostEditChangeProperty( struct FPropertyChangedEvent& PropertyChangedEvent )
+{
+	Super::PostEditChangeProperty(PropertyChangedEvent);
+
+	const FName Name = (PropertyChangedEvent.Property != nullptr) ? PropertyChangedEvent.Property->GetFName() : NAME_None;
+	SettingChangedEvent.Broadcast(Name);
+}
diff --git a/Engine/Source/Editor/NewsFeed/Private/NewsFeedModule.cpp b/Engine/Source/Editor/NewsFeed/Private/NewsFeedModule.cpp
new file mode 100644
--- /dev/null
+++ b/Engine/Source/Editor/NewsFeed/Private/NewsFeedModule.cpp
@@ -0,0 +1,152 @@
+// Copyright 1998-2014 Epic Games, Inc. All Rights Reserved.
+
+/*=============================================================================
+	NewsFeedModule.cpp: Implements the FNewsFeedModule class.
+=============================================================================*/
+
+#include "NewsFeedPrivatePCH.h"
+
+
+#define LOCTEXT_NAMESPACE "FNewsFeedModule"
+
+
+/**
+ * Implements the NewsFeed module.
+ */
+class FNewsFeedModule
+	: public INewsFeedModule
+{
+public:
+
+	// Begin IModuleInterface interface
+
+	virtual void StartupModule( ) OVERRIDE
+	{
+		// register settings
+		ISettingsModule* SettingsModule = ISettingsModule::Get();
+
+		if (SettingsModule != nullptr)
+		{
+			SettingsModule->RegisterSettings("Editor", "LevelEditor", "NewsFeed",
+				LOCTEXT("NewsFeedSettingsName", "News Feed"),
+				LOCTEXT("NewsFeedSettingsDescription", "Change how you receive and view news updates."),
+				TWeakObjectPtr<UObject>(UNewsFeedSettings::StaticClass()->GetDefaultObject())
+			);
+		}
+	}
+
+	virtual void ShutdownModule( ) OVERRIDE
+	{
+		// unregister settings
+		ISettingsModule* SettingsModule = ISettingsModule::Get();
+
+		if (SettingsModule != nullptr)
+		{
+			SettingsModule->UnregisterSettings("Editor", "LevelEditor", "NewsFeed");
+		}
+	}
+
+	// End IModuleInterface interface
+
+public:
+
+	// Begin INewsFeedModule interface
+
+	virtual TSharedRef<class SWidget> CreateNewsFeedButton( ) OVERRIDE
+	{
+		FNewsFeedCacheRef NewsFeedCache = MakeShareable(new FNewsFeedCache());
+		FNewsFeedCachePtr Bleh = NewsFeedCache;
+
+		TSharedPtr<STextBlock> UnreadItemsTextBlock;
+
+		return SNew(SBox)
+			.Padding(FMargin(0.0f, 0.0f, 0.0f, 1.0f))
+			.VAlign(VAlign_Bottom)
+			[
+				SNew(SComboButton)
+					.HasDownArrow(false)
+					.ButtonStyle(FEditorStyle::Get(), "NoBorder")
+					.ContentPadding(0.0f)
+					.MenuPlacement(MenuPlacement_BelowAnchor)
+					.ButtonContent()
+					[
+						SNew(SOverlay)
+
+						+ SOverlay::Slot()
+							.HAlign(HAlign_Center)
+							.VAlign(VAlign_Bottom)
+							.Padding(FMargin(0.0f, 0.0f, 3.0f, 0.0f))
+							[
+								SNew(SImage)
+									.Image(FEditorStyle::GetBrush("NewsFeed.ToolbarIcon.Small"))
+							]
+
+						+ SOverlay::Slot()
+							.HAlign(HAlign_Right)
+							.VAlign(VAlign_Top)
+							[
+								SNew(SBorder)
+									.BorderBackgroundColor_Raw(this, &FNewsFeedModule::HandleUnreadCountBackgroundColor)
+									.BorderImage(FEditorStyle::GetBrush("NewsFeed.UnreadCountBackground"))
+									.HAlign(HAlign_Center)
+									.VAlign(VAlign_Center)
+									.Padding(FMargin(2.0f, 0.0f))
+									.Visibility_Raw(this, &FNewsFeedModule::HandleUnreadItemsVisibility, UnreadItemsTextBlock)
+									[
+										SAssignNew(UnreadItemsTextBlock, STextBlock)
+											.ColorAndOpacity(FLinearColor(0.0f, 0.0f, 0.0f))
+											.Font(FSlateFontInfo(FPaths::EngineContentDir() / TEXT("Slate/Fonts/Roboto-Bold.ttf"), 7))
+											.Text_Raw(this, &FNewsFeedModule::HandleUnreadItemsText, Bleh)
+									]
+							]
+					]
+					.MenuContent()
+					[
+						SNew(SNewsFeed, NewsFeedCache)
+					]
+			];
+	}
+
+	// End INewsFeedModule interface
+
+private:
+
+	// Callback for getting the background color the unread item count.
+	FSlateColor HandleUnreadCountBackgroundColor( ) const
+	{
+		return GetDefault<UEditorStyleSettings>()->SelectionColor;
+	}
+
+	// Callback for getting the text of the unread item count.
+	FText HandleUnreadItemsText( FNewsFeedCachePtr NewsFeedCache ) const
+	{
+		int32 UnreadItems = 0;
+
+		for (auto Item : NewsFeedCache->GetItems())
+		{
+			if (!Item->Read)
+			{
+				++UnreadItems;
+			}
+		}
+
+		if (UnreadItems > 0)
+		{
+			return FText::AsNumber(UnreadItems);
+		}
+
+		return FText::GetEmpty();
+	}
+
+	// Callback for getting the visibility of the unread item count.
+	EVisibility HandleUnreadItemsVisibility( TSharedPtr<STextBlock> UnreadItemsTextBlock ) const
+	{
+		return UnreadItemsTextBlock->GetText().IsEmpty() ? EVisibility::Hidden : EVisibility::Visible;
+	}
+};
+
+
+IMPLEMENT_MODULE(FNewsFeedModule, NewsFeed);
+
+
+#undef LOCTEXT_NAMESPACE
diff --git a/Engine/Source/Editor/NewsFeed/Private/NewsFeedPrivatePCH.h b/Engine/Source/Editor/NewsFeed/Private/NewsFeedPrivatePCH.h
new file mode 100644
--- /dev/null
+++ b/Engine/Source/Editor/NewsFeed/Private/NewsFeedPrivatePCH.h
@@ -0,0 +1,46 @@
+// Copyright 1998-2014 Epic Games, Inc. All Rights Reserved.
+
+/*=============================================================================
+	NewsFeedPrivatePCH.h: Pre-compiled header file for the NewsFeed module.
+=============================================================================*/
+
+#pragma once
+
+#include "../Public/NewsFeed.h"
+
+
+/* Dependencies
+ *****************************************************************************/
+
+#include "AnalyticsEventAttribute.h"
+#include "Base64.h"
+#include "CoreUObject.h"
+#include "DesktopPlatformModule.h"
+#include "EngineAnalytics.h"
+#include "HttpModule.h"
+#include "IAnalyticsProvider.h"
+#include "IHttpBase.h"
+#include "IHttpRequest.h"
+#include "IHttpResponse.h"
+#include "ImageWrapper.h"
+#include "Json.h"
+#include "ModuleManager.h"
+#include "Online.h"
+#include "OnlineSubsystemTypes.h"
+#include "OnlineTitleFileInterface.h"
+#include "SecureHash.h"
+#include "Settings.h"
+#include "Ticker.h"
+
+
+/* Private includes
+ *****************************************************************************/
+
+#include "CdnNewsFeedTitleFile.h"
+#include "LocalNewsFeedTitleFile.h"
+#include "NewsFeedItem.h"
+#include "NewsFeedSettings.h"
+#include "NewsFeedCache.h"
+
+#include "SNewsFeedListRow.h"
+#include "SNewsFeed.h"
diff --git a/Engine/Source/Editor/NewsFeed/Private/Widgets/SNewsFeed.cpp b/Engine/Source/Editor/NewsFeed/Private/Widgets/SNewsFeed.cpp
new file mode 100644
--- /dev/null
+++ b/Engine/Source/Editor/NewsFeed/Private/Widgets/SNewsFeed.cpp
@@ -0,0 +1,351 @@
+// Copyright 1998-2014 Epic Games, Inc. All Rights Reserved.
+
+/*=============================================================================
+	SNewsFeed.cpp: Implements the SNewsFeed class.
+=============================================================================*/
+
+#include "NewsFeedPrivatePCH.h"
+
+
+#define LOCTEXT_NAMESPACE "SNewsFeed"
+
+
+/* SNewsFeed structors
+ *****************************************************************************/
+
+SNewsFeed::~SNewsFeed( )
+{
+	GetMutableDefault<UNewsFeedSettings>()->OnSettingChanged().RemoveAll(this);
+
+	NewsFeedCache->OnLoadCompleted().RemoveAll(this);
+	NewsFeedCache->OnLoadFailed().RemoveAll(this);
+}
+
+
+/* SNewsFeed interface
+ *****************************************************************************/
+
+void SNewsFeed::Construct( const FArguments& InArgs, const FNewsFeedCacheRef& InNewsFeedCache )
+{
+	NewsFeedCache = InNewsFeedCache;
+
+	ChildSlot
+	[
+		SNew(SVerticalBox)
+
+		+ SVerticalBox::Slot()
+			.AutoHeight()
+			.Padding(4.0f)
+			[
+				SNew(SOverlay)
+
+				+ SOverlay::Slot()
+					[
+						// news list
+						SAssignNew(NewsFeedItemListView, SListView<FNewsFeedItemPtr>)
+							.ItemHeight(64.0f)
+							.ListItemsSource(&NewsFeedItemList)
+							.OnGenerateRow(this, &SNewsFeed::HandleNewsListViewGenerateRow)
+							.OnSelectionChanged(this, &SNewsFeed::HandleNewsListViewSelectionChanged)
+							.SelectionMode(ESelectionMode::Single)
+							.Visibility(this, &SNewsFeed::HandleNewsListViewVisibility)
+							.HeaderRow
+							(
+								SNew(SHeaderRow)
+									.Visibility(EVisibility::Collapsed)
+
+								+ SHeaderRow::Column("Icon")
+									.FixedWidth(48.0f)
+
+								+ SHeaderRow::Column("Content")
+							)
+					]
+
+				+ SOverlay::Slot()
+					.HAlign(HAlign_Center)
+					.Padding(FMargin(16.0f, 4.0f))
+					[
+						// no news available
+						SNew(STextBlock)
+							.Text(this, &SNewsFeed::HandleNewsUnavailableText)
+							.Visibility(this, &SNewsFeed::HandleNewsUnavailableVisibility)
+					]
+			]
+
+		+ SVerticalBox::Slot()
+			.AutoHeight()
+			[
+				SNew(SSeparator)
+			]
+
+		+ SVerticalBox::Slot()
+			.AutoHeight()
+			.Padding(4.0f)
+			[
+				SNew(SHorizontalBox)
+
+				+ SHorizontalBox::Slot()
+					.AutoWidth()
+					.HAlign(HAlign_Center)
+					.VAlign(VAlign_Center)
+					.Padding(4.0f, 0.0f)
+					[
+						// settings button
+						SNew(SButton)
+							.ButtonStyle(FEditorStyle::Get(), "NoBorder")
+							.ContentPadding(0.0f)
+							.OnClicked(this, &SNewsFeed::HandleRefreshButtonClicked)
+							.ToolTipText(LOCTEXT("RefreshButtonToolTip", "Reload the latest news from the server."))
+							.Visibility(this, &SNewsFeed::HandleRefreshButtonVisibility)
+							[
+								SNew(SImage)
+									.Image(FEditorStyle::GetBrush("NewsFeed.ReloadButton"))
+							]
+					]
+
+				+ SHorizontalBox::Slot()
+					.FillWidth(1.0f)
+					.HAlign(HAlign_Center)
+					[
+						// mark all as read
+						SNew(SHyperlink)
+							.IsEnabled(this, &SNewsFeed::HandleMarkAllAsReadIsEnabled)
+							.OnNavigate(this, &SNewsFeed::HandleMarkAllAsReadNavigate)
+							.Text(LOCTEXT("MarkAllAsReadHyperlink", "Mark all news as read"))
+					]
+
+				+ SHorizontalBox::Slot()
+					.AutoWidth()
+					.HAlign(HAlign_Center)
+					.VAlign(VAlign_Center)
+					.Padding(4.0f, 0.0f)
+					[
+						SNew(SOverlay)
+
+						+ SOverlay::Slot()
+							.HAlign(HAlign_Center)
+							.VAlign(VAlign_Center)
+							[
+								// reload button
+								SNew(SButton)
+									.ButtonStyle(FEditorStyle::Get(), "NoBorder")
+									.ContentPadding(0.0f)
+									.OnClicked(this, &SNewsFeed::HandleSettingsButtonClicked)
+									.ToolTipText(LOCTEXT("SettingsButtonToolTip", "News feed settings."))
+									[
+										SNew(SImage)
+											.Image(FEditorStyle::GetBrush("NewsFeed.SettingsButton"))
+									]
+							]
+
+						+ SOverlay::Slot()
+							.HAlign(HAlign_Center)
+							.VAlign(VAlign_Center)
+							[
+								// refresh throbber
+								SNew(SCircularThrobber)
+									.Radius(10.0f)
+									.ToolTipText(LOCTEXT("RefreshThrobberToolTip", "Loading latest news..."))
+									.Visibility(this, &SNewsFeed::HandleRefreshThrobberVisibility)
+							]
+					]
+			]
+	];
+
+	NewsFeedCache->OnLoadCompleted().AddRaw(this, &SNewsFeed::HandleNewsFeedLoaderLoadCompleted);
+	NewsFeedCache->OnLoadFailed().AddRaw(this, &SNewsFeed::HandleNewsFeedLoaderLoadFailed);
+	NewsFeedCache->LoadFeed();
+
+	GetMutableDefault<UNewsFeedSettings>()->OnSettingChanged().AddSP(this, &SNewsFeed::HandleNewsFeedSettingsChanged);
+}
+
+
+/* SNewsFeed implementation
+ *****************************************************************************/
+
+void SNewsFeed::ReloadNewsFeedItems( )
+{
+	NewsFeedItemList.Empty();
+
+	const UNewsFeedSettings* NewsFeedSettings = GetDefault<UNewsFeedSettings>();
+
+	for (auto NewsFeedItem : NewsFeedCache->GetItems())
+	{
+		if (!GetDefault<UNewsFeedSettings>()->ShowOnlyUnreadItems || !NewsFeedItem->Read)
+		{
+			NewsFeedItemList.Add(NewsFeedItem);
+
+			if (NewsFeedItemList.Num() == NewsFeedSettings->MaxItemsToShow)
+			{
+				break;
+			}
+		}
+	}
+
+	NewsFeedItemListView->RequestListRefresh();
+}
+
+
+/* SNewsFeed callbacks
+ *****************************************************************************/
+
+EVisibility SNewsFeed::HandleLoadFailedBoxVisibility( ) const
+{
+	return NewsFeedCache->IsFinished() ? EVisibility::Collapsed : EVisibility::Visible;
+}
+
+
+bool SNewsFeed::HandleMarkAllAsReadIsEnabled( ) const
+{
+	for (auto NewsFeedItem : NewsFeedItemList)
+	{
+		if (!NewsFeedItem->Read)
+		{
+			return true;
+		}
+	}
+
+	return false;
+}
+
+
+void SNewsFeed::HandleMarkAllAsReadNavigate( )
+{
+	for (auto NewsFeedItem : NewsFeedItemList)
+	{
+		NewsFeedItem->Read = true;
+	}
+
+	if (GetDefault<UNewsFeedSettings>()->ShowOnlyUnreadItems)
+	{
+		ReloadNewsFeedItems();
+	}
+}
+
+
+void SNewsFeed::HandleNewsFeedLoaderLoadCompleted( )
+{
+	ReloadNewsFeedItems();
+}
+
+
+void SNewsFeed::HandleNewsFeedLoaderLoadFailed( )
+{
+
+}
+
+
+void SNewsFeed::HandleNewsFeedSettingsChanged( FName PropertyName )
+{
+	ReloadNewsFeedItems();
+}
+
+
+TSharedRef<ITableRow> SNewsFeed::HandleNewsListViewGenerateRow( FNewsFeedItemPtr NewsFeedItem, const TSharedRef<STableViewBase>& OwnerTable )
+{
+	return SNew(SNewsFeedListRow, OwnerTable, NewsFeedCache.ToSharedRef())
+		.NewsFeedItem(NewsFeedItem);
+}
+
+
+void SNewsFeed::HandleNewsListViewSelectionChanged( FNewsFeedItemPtr Selection, ESelectInfo::Type SelectInfo )
+{
+	if (Selection.IsValid() && (SelectInfo == ESelectInfo::OnMouseClick))
+	{
+		TArray<FAnalyticsEventAttribute> EventAttributes;
+
+		EventAttributes.Add(FAnalyticsEventAttribute(TEXT("ItemID"), Selection->ItemId.ToString()));
+
+		IDesktopPlatform* DesktopPlatform = FDesktopPlatformModule::Get();
+
+		if (DesktopPlatform != nullptr)
+		{
+			bool Succeeded = false;
+
+			if (Selection->Url.StartsWith(TEXT("http://")))
+			{
+				FPlatformProcess::LaunchURL(*Selection->Url, NULL, NULL);
+				Succeeded = true;
+			}
+			else if (Selection->Url.StartsWith(TEXT("ue4://market/")))
+			{
+				Succeeded = DesktopPlatform->OpenLauncher(false, FString::Printf(TEXT("-OpenMarketItem=%s"), *Selection->Url.RightChop(13)));
+			}
+//			else if (Selection->Url.StartsWith(TEXT("ue4://")))
+//			{
+//				Succeeded = DesktopPlatform->OpenLauncher(false, FString::Printf(TEXT("-OpenUrl=%s"), *Selection->Url));
+//			}
+
+			EventAttributes.Add(FAnalyticsEventAttribute(TEXT("OpenUrlSucceeded"), Succeeded ? TEXT("TRUE") : TEXT("FALSE")));
+
+			if (FEngineAnalytics::IsAvailable())
+			{
+				FEngineAnalytics::GetProvider().RecordEvent(TEXT("Editor.Usage.OpenNewsFeedItem"), EventAttributes);
+			}
+		}
+
+		// hack: mark as read
+		Selection->Read = true;
+		
+		NewsFeedItemListView->SetSelection(nullptr);
+
+		if (GetDefault<UNewsFeedSettings>()->ShowOnlyUnreadItems)
+		{
+			ReloadNewsFeedItems();
+		}
+	}
+}
+
+
+FText SNewsFeed::HandleNewsUnavailableText( ) const
+{
+	if (GetDefault<UNewsFeedSettings>()->ShowOnlyUnreadItems)
+	{
+		return LOCTEXT("NoUnreadNewsAvailable", "There are currently no unread news");
+	}
+
+	return LOCTEXT("NoNewsAvailable", "There are currently no news");
+}
+
+
+EVisibility SNewsFeed::HandleNewsListViewVisibility( ) const
+{
+	return (NewsFeedItemList.Num() == 0) ? EVisibility::Collapsed : EVisibility::Visible;
+}
+
+
+EVisibility SNewsFeed::HandleNewsUnavailableVisibility( ) const
+{
+	return (NewsFeedItemList.Num() > 0) ? EVisibility::Collapsed : EVisibility::Visible;
+}
+
+
+FReply SNewsFeed::HandleRefreshButtonClicked( ) const
+{
+	NewsFeedCache->LoadFeed();
+
+	return FReply::Handled();
+}
+
+
+EVisibility SNewsFeed::HandleRefreshButtonVisibility( ) const
+{
+	return NewsFeedCache->IsLoading() ? EVisibility::Hidden : EVisibility::Visible;
+}
+
+
+EVisibility SNewsFeed::HandleRefreshThrobberVisibility( ) const
+{
+	return NewsFeedCache->IsLoading() ? EVisibility::Visible : EVisibility::Hidden;
+}
+
+
+FReply SNewsFeed::HandleSettingsButtonClicked( ) const
+{
+	FModuleManager::LoadModuleChecked<ISettingsModule>("Settings").ShowViewer("Editor", "LevelEditor", "NewsFeed");
+
+	return FReply::Handled();
+}
+
+
+#undef LOCTEXT_NAMESPACE
diff --git a/Engine/Source/Editor/NewsFeed/Private/Widgets/SNewsFeed.h b/Engine/Source/Editor/NewsFeed/Private/Widgets/SNewsFeed.h
new file mode 100644
--- /dev/null
+++ b/Engine/Source/Editor/NewsFeed/Private/Widgets/SNewsFeed.h
@@ -0,0 +1,102 @@
+// Copyright 1998-2014 Epic Games, Inc. All Rights Reserved.
+
+/*=============================================================================
+	SNewsFeed.h: Declares the SNewsFeed class.
+=============================================================================*/
+
+#pragma once
+
+
+/**
+ * Implements the news feed widget.
+ */
+class SNewsFeed
+	: public SCompoundWidget
+{
+public:
+
+	SLATE_BEGIN_ARGS(SNewsFeed) { }
+	SLATE_END_ARGS()
+
+public:
+
+	/**
+	 * Destructor.
+	 */
+	~SNewsFeed( );
+
+public:
+
+	/**
+	 * Constructs the widget.
+	 *
+	 * @param InArgs - The construction arguments.
+	 * @param InNewsFeedCache - The news feed cache to use.
+	 */
+	void Construct( const FArguments& InArgs, const FNewsFeedCacheRef& InNewsFeedCache );
+
+protected:
+
+	/**
+	 * Reloads the list of news feed items.
+	 */
+	void ReloadNewsFeedItems( );
+
+private:
+
+	// Callback for getting the visibility of the 'Loading failed' box.
+	EVisibility HandleLoadFailedBoxVisibility( ) const;
+
+	// Callback for getting the enabled state of the 'Mark all as read' hyper-link.
+	bool HandleMarkAllAsReadIsEnabled( ) const;
+
+	// Callback for clicking the 'Mark all as read' hyper-link.
+	void HandleMarkAllAsReadNavigate( );
+
+	// Callback for completed news feeds loads.
+	void HandleNewsFeedLoaderLoadCompleted( );
+
+	// Callback for failed news feeds loads.
+	void HandleNewsFeedLoaderLoadFailed( );
+
+	// Callback for news feed settings changes.
+	void HandleNewsFeedSettingsChanged( FName PropertyName );
+
+	// Callback for generating a row widget in the device service list view.
+	TSharedRef<ITableRow> HandleNewsListViewGenerateRow( FNewsFeedItemPtr NewsFeedItem, const TSharedRef<STableViewBase>& OwnerTable );
+
+	// Callback for selecting a news feed item.
+	void HandleNewsListViewSelectionChanged( FNewsFeedItemPtr Selection, ESelectInfo::Type SelectInfo );
+
+	// Callback for getting the visibility of the news list view.
+	EVisibility HandleNewsListViewVisibility( ) const;
+
+	// Callback for getting the text of the 'news unavailable' notification.
+	FText HandleNewsUnavailableText( ) const;
+
+	// Callback for getting the visibility of the 'news unavailable' notification.
+	EVisibility HandleNewsUnavailableVisibility( ) const;
+
+	// Callback for clicking the 'Refresh' button.
+	FReply HandleRefreshButtonClicked( ) const;
+
+	// Callback for getting the visibility of the 'Refresh' button.
+	EVisibility HandleRefreshButtonVisibility( ) const;
+
+	// Callback for getting the visibility of the refresh throbber.
+	EVisibility HandleRefreshThrobberVisibility( ) const;
+
+	// Callback for clicking the 'Settings' button.
+	FReply HandleSettingsButtonClicked( ) const;
+
+private:
+
+	// Holds a pointer the news feed cache.
+	FNewsFeedCachePtr NewsFeedCache;
+
+	// Holds the list of news items.
+	TArray<FNewsFeedItemPtr> NewsFeedItemList;
+
+	// Holds the list view for news items.
+	TSharedPtr<SListView<FNewsFeedItemPtr> > NewsFeedItemListView;
+};
diff --git a/Engine/Source/Editor/NewsFeed/Private/Widgets/SNewsFeedListRow.h b/Engine/Source/Editor/NewsFeed/Private/Widgets/SNewsFeedListRow.h
new file mode 100644
--- /dev/null
+++ b/Engine/Source/Editor/NewsFeed/Private/Widgets/SNewsFeedListRow.h
@@ -0,0 +1,141 @@
+// Copyright 1998-2014 Epic Games, Inc. All Rights Reserved.
+
+/*=============================================================================
+	SNewsFeedListRow.h: Declares the SNewsFeedListRow class.
+=============================================================================*/
+
+#pragma once
+
+
+#define LOCTEXT_NAMESPACE "SNewsFeedListRow"
+
+
+/**
+ * Implements a row widget for the news list view.
+ */
+class SNewsFeedListRow
+	: public SMultiColumnTableRow<FNewsFeedItemPtr>
+{
+public:
+
+	SLATE_BEGIN_ARGS(SNewsFeedListRow) { }
+		SLATE_ARGUMENT(FNewsFeedItemPtr, NewsFeedItem)
+	SLATE_END_ARGS()
+
+public:
+
+	/**
+	 * Constructs the widget.
+	 *
+	 * @param InArgs - The arguments.
+	 */
+	void Construct( const FArguments& InArgs, const TSharedRef<STableViewBase>& InOwnerTableView, const FNewsFeedCacheRef& InNewsFeedCache )
+	{
+		check(InArgs._NewsFeedItem.IsValid());
+
+		NewsFeedCache = InNewsFeedCache;
+		NewsFeedItem = InArgs._NewsFeedItem;
+		
+		SMultiColumnTableRow<FNewsFeedItemPtr>::Construct(FSuperRowType::FArguments(), InOwnerTableView);
+	}
+
+public:
+
+	/**
+	 * Generates the widget for the specified column.
+	 *
+	 * @param ColumnName - The name of the column to generate the widget for.
+	 *
+	 * @return The widget.
+	 */
+	BEGIN_SLATE_FUNCTION_BUILD_OPTIMIZATION
+	virtual TSharedRef<SWidget> GenerateWidgetForColumn( const FName& ColumnName ) OVERRIDE
+	{
+		if (ColumnName == TEXT("Icon"))
+		{
+			return SNew(SBox)
+				.HAlign(HAlign_Left)
+				.VAlign(VAlign_Top)
+				.Padding(4.0f)
+				[
+					SNew(SImage)
+						.Image(this, &SNewsFeedListRow::HandleIconImage)
+				];
+		}
+		else if (ColumnName == TEXT("Content"))
+		{
+			return SNew(SVerticalBox)
+				.ToolTipText(NewsFeedItem->FullText)
+
+			+ SVerticalBox::Slot()
+				.AutoHeight()
+				.Padding(4.0f)
+				[
+					SNew(SHorizontalBox)
+
+					+ SHorizontalBox::Slot()
+						.FillWidth(1.0f)
+						[
+							SNew(STextBlock)
+								.ColorAndOpacity(this, &SNewsFeedListRow::HandleTitleColorAndOpacity)
+								.Font(FSlateFontInfo(FPaths::EngineContentDir() / TEXT("Slate/Fonts/Roboto-Bold.ttf"), 12))
+								.Text(NewsFeedItem->Title)							
+						]
+
+					+ SHorizontalBox::Slot()
+						.AutoWidth()
+						.Padding(24.0f, 0.0f, 0.0f, 0.0f)
+						.VAlign(VAlign_Top)
+						[
+							SNew(STextBlock)
+								.ColorAndOpacity(FSlateColor::UseSubduedForeground())
+								.Font(FSlateFontInfo(FPaths::EngineContentDir() / TEXT("Slate/Fonts/Roboto-Regular.ttf"), 8))
+								.Text(FText::AsDate(NewsFeedItem->Issued.ToDate()))
+						]
+				]
+				
+			+ SVerticalBox::Slot()
+				.AutoHeight()
+				.Padding(4.0f, 0.0f, 4.0f, 8.0f)
+				[
+					SNew(STextBlock)
+						.AutoWrapText(true)
+						.ColorAndOpacity(FSlateColor::UseSubduedForeground())
+						.Text(NewsFeedItem->Excerpt)
+				];
+		}
+
+		return SNullWidget::NullWidget;
+	}
+	END_SLATE_FUNCTION_BUILD_OPTIMIZATION
+
+private:
+
+	// Callback for getting the image brush of the icon.
+	const FSlateBrush* HandleIconImage( ) const
+	{
+		return NewsFeedCache->GetIcon(NewsFeedItem->IconName);
+	}
+
+	// Callback for getting the color of the title text.
+	FSlateColor HandleTitleColorAndOpacity() const
+	{
+		if (NewsFeedItem->Read)
+		{
+			return FSlateColor::UseSubduedForeground();
+		}
+
+		return FSlateColor::UseForeground();
+	}
+
+private:
+
+	// Holds a pointer to the news feed cache.
+	FNewsFeedCachePtr NewsFeedCache;
+
+	// Holds the news item shown in this row.
+	FNewsFeedItemPtr NewsFeedItem;
+};
+
+
+#undef LOCTEXT_NAMESPACE
diff --git a/Engine/Source/Editor/NewsFeed/Public/Interfaces/INewsFeedModule.h b/Engine/Source/Editor/NewsFeed/Public/Interfaces/INewsFeedModule.h
new file mode 100644
--- /dev/null
+++ b/Engine/Source/Editor/NewsFeed/Public/Interfaces/INewsFeedModule.h
@@ -0,0 +1,31 @@
+// Copyright 1998-2014 Epic Games, Inc. All Rights Reserved.
+
+/*=============================================================================
+	INewsFeedModule.h: Declares the INewsFeedModule interface.
+=============================================================================*/
+
+#pragma once
+
+
+/**
+ * Interface for the news feed module.
+ */
+class INewsFeedModule
+	: public IModuleInterface
+{
+public:
+
+	/**
+	 * Creates a news feed button widget.
+	 *
+	 * @return The new widget.
+	 */
+	virtual TSharedRef<class SWidget> CreateNewsFeedButton( ) = 0;
+
+public:
+
+	/**
+	 * Virtual destructor.
+	 */
+	virtual ~INewsFeedModule( ) { }
+};
diff --git a/Engine/Source/Editor/NewsFeed/Public/NewsFeed.h b/Engine/Source/Editor/NewsFeed/Public/NewsFeed.h
new file mode 100644
--- /dev/null
+++ b/Engine/Source/Editor/NewsFeed/Public/NewsFeed.h
@@ -0,0 +1,22 @@
+// Copyright 1998-2014 Epic Games, Inc. All Rights Reserved.
+
+/*=============================================================================
+	NewsFeed.h: NewsFeed module public header file.
+=============================================================================*/
+
+#pragma once
+
+
+/* Dependencies
+ *****************************************************************************/
+
+#include "Core.h"
+#include "ModuleInterface.h"
+#include "Slate.h"
+#include "EditorStyle.h"
+
+
+/* Interfaces
+ *****************************************************************************/
+
+#include "INewsFeedModule.h"
diff --git a/Engine/Source/Runtime/Online/OnlineSubsystem/Public/Interfaces/OnlineTitleFileInterface.h b/Engine/Source/Runtime/Online/OnlineSubsystem/Public/Interfaces/OnlineTitleFileInterface.h
--- a/Engine/Source/Runtime/Online/OnlineSubsystem/Public/Interfaces/OnlineTitleFileInterface.h
+++ b/Engine/Source/Runtime/Online/OnlineSubsystem/Public/Interfaces/OnlineTitleFileInterface.h
@@ -101,4 +101,5 @@ public:
 
 };
 
-typedef TSharedPtr<IOnlineTitleFile, ESPMode::ThreadSafe> IOnlineTitleFilePtr;
\ No newline at end of file
+typedef TSharedPtr<IOnlineTitleFile, ESPMode::ThreadSafe> IOnlineTitleFilePtr;
+typedef TSharedRef<IOnlineTitleFile, ESPMode::ThreadSafe> IOnlineTitleFileRef;
\ No newline at end of file
