fixed News Feed widget crashing on MacOS

diff --git a/Engine/Source/Editor/NewsFeed/Private/NewsFeedModule.cpp b/Engine/Source/Editor/NewsFeed/Private/NewsFeedModule.cpp
--- a/Engine/Source/Editor/NewsFeed/Private/NewsFeedModule.cpp
+++ b/Engine/Source/Editor/NewsFeed/Private/NewsFeedModule.cpp
@@ -55,9 +55,11 @@ public:
 	virtual TSharedRef<class SWidget> CreateNewsFeedButton( ) OVERRIDE
 	{
 		FNewsFeedCacheRef NewsFeedCache = MakeShareable(new FNewsFeedCache());
-		FNewsFeedCachePtr Bleh = NewsFeedCache;
 
-		TSharedPtr<STextBlock> UnreadItemsTextBlock;
+		TSharedRef<STextBlock> UnreadItemsTextBlock = SNew(STextBlock)
+			.ColorAndOpacity(FLinearColor(0.0f, 0.0f, 0.0f))
+			.Font(FSlateFontInfo(FPaths::EngineContentDir() / TEXT("Slate/Fonts/Roboto-Bold.ttf"), 7))
+			.Text_Raw(this, &FNewsFeedModule::HandleUnreadItemsText, NewsFeedCache);
 
 		return SNew(SBox)
 			.Padding(FMargin(0.0f, 0.0f, 0.0f, 1.0f))
@@ -93,10 +95,7 @@ public:
 									.Padding(FMargin(2.0f, 0.0f))
 									.Visibility_Raw(this, &FNewsFeedModule::HandleUnreadItemsVisibility, UnreadItemsTextBlock)
 									[
-										SAssignNew(UnreadItemsTextBlock, STextBlock)
-											.ColorAndOpacity(FLinearColor(0.0f, 0.0f, 0.0f))
-											.Font(FSlateFontInfo(FPaths::EngineContentDir() / TEXT("Slate/Fonts/Roboto-Bold.ttf"), 7))
-											.Text_Raw(this, &FNewsFeedModule::HandleUnreadItemsText, Bleh)
+										UnreadItemsTextBlock
 									]
 							]
 					]
@@ -118,7 +117,7 @@ private:
 	}
 
 	// Callback for getting the text of the unread item count.
-	FText HandleUnreadItemsText( FNewsFeedCachePtr NewsFeedCache ) const
+	FText HandleUnreadItemsText( FNewsFeedCacheRef NewsFeedCache ) const
 	{
 		int32 UnreadItems = 0;
 
@@ -139,7 +138,7 @@ private:
 	}
 
 	// Callback for getting the visibility of the unread item count.
-	EVisibility HandleUnreadItemsVisibility( TSharedPtr<STextBlock> UnreadItemsTextBlock ) const
+	EVisibility HandleUnreadItemsVisibility( TSharedRef<STextBlock> UnreadItemsTextBlock ) const
 	{
 		return UnreadItemsTextBlock->GetText().IsEmpty() ? EVisibility::Hidden : EVisibility::Visible;
 	}
