TTP# 332978: NewsFeed message reading is not a global setting; fixed

diff --git a/Engine/Source/Editor/NewsFeed/Private/Implementation/NewsFeedSettings.h b/Engine/Source/Editor/NewsFeed/Private/Implementation/NewsFeedSettings.h
--- a/Engine/Source/Editor/NewsFeed/Private/Implementation/NewsFeedSettings.h
+++ b/Engine/Source/Editor/NewsFeed/Private/Implementation/NewsFeedSettings.h
@@ -9,7 +9,7 @@
 #include "NewsFeedSettings.generated.h"
 
 
-UCLASS(config=Editor)
+UCLASS(config=EditorGameAgnostic)
 class UNewsFeedSettings
 	: public UObject
 {
