NewsFeed: moved title file configuration settings into NewsFeedSettings object

diff --git a/Engine/Config/BaseEditor.ini b/Engine/Config/BaseEditor.ini
--- a/Engine/Config/BaseEditor.ini
+++ b/Engine/Config/BaseEditor.ini
@@ -253,11 +253,6 @@ NotificationDuration=10
 PulseTimeInterval=10
 PulseDuration=5
 
-[NewsFeed]
-CdnUrl="http://cdn.unrealengine.com/newsfeed/{Culture}/index.json"
-EngineContentRelativeDirectory="Editor/NewsFeed/{Culture}/"
-Source="CDN"
-
 [CustomBlueprintFunctionLibrary]
 bCanCreateNew=true
 
diff --git a/Engine/Config/BaseEditorGameAgnostic.ini b/Engine/Config/BaseEditorGameAgnostic.ini
--- a/Engine/Config/BaseEditorGameAgnostic.ini
+++ b/Engine/Config/BaseEditorGameAgnostic.ini
@@ -1 +1,8 @@
 [/Script/UnrealEd.EditorGameAgnosticSettings]
+
+[/Script/NewsFeed.NewsFeedSettings]
+CdnSourceUrl="http://cdn.unrealengine.com/newsfeed/{Culture}/index.json"
+LocalSourcePath="Editor/NewsFeed/{Culture}/"
+Source=NEWSFEED_Cdn
+MaxItemsToShow=10
+ShowOnlyUnreadItems=True
\ No newline at end of file
diff --git a/Engine/Source/Editor/NewsFeed/Private/Implementation/NewsFeedCache.cpp b/Engine/Source/Editor/NewsFeed/Private/Implementation/NewsFeedCache.cpp
--- a/Engine/Source/Editor/NewsFeed/Private/Implementation/NewsFeedCache.cpp
+++ b/Engine/Source/Editor/NewsFeed/Private/Implementation/NewsFeedCache.cpp
@@ -101,54 +101,50 @@ void FNewsFeedCache::LoadTitleFile( )
 	}
 
 	// get title file source
-	FString SourceType;
+	const UNewsFeedSettings* NewsFeedSettings = GetDefault<UNewsFeedSettings>();
 
-	if (GConfig->GetString(TEXT("NewsFeed"), TEXT("Source"), SourceType, GEditorIni) && !SourceType.IsEmpty())
+	if (NewsFeedSettings->Source == NEWSFEED_Cdn)
 	{
-		if (SourceType == TEXT("CDN"))
+		if (!NewsFeedSettings->CdnSourceUrl.IsEmpty())
 		{
-			FString CdnUrl;
+			FString CdnSourceUrl;
 
-			if (GConfig->GetString(TEXT("NewsFeed"),TEXT("CdnUrl"), CdnUrl, GEditorIni) && !CdnUrl.IsEmpty())
+			if (!CultureString.IsEmpty())
 			{
-				if (!CultureString.IsEmpty())
-				{
-					CdnUrl = CdnUrl.Replace( TEXT("{Culture}"), *CultureString );
-				}
-				else
-				{
-					CdnUrl = CdnUrl.Replace(TEXT("{Culture}/"), TEXT(""));
-				}
-
-				TitleFile = FCdnNewsFeedTitleFile::Create(CdnUrl);
+				CdnSourceUrl = NewsFeedSettings->CdnSourceUrl.Replace(TEXT("{Culture}"), *CultureString);
 			}
+			else
+			{
+				CdnSourceUrl = NewsFeedSettings->CdnSourceUrl.Replace(TEXT("{Culture}/"), TEXT(""));
+			}
+
+			TitleFile = FCdnNewsFeedTitleFile::Create(CdnSourceUrl);
 		}
-		else if (SourceType == TEXT("LOCAL"))
+	}
+	else if (NewsFeedSettings->Source == NEWSFEED_Local)
+	{
+		if (!NewsFeedSettings->LocalSourcePath.IsEmpty())
 		{
-			FString EngineContentRelativeDirectory;
+			FString RootDirectory = FPaths::Combine(FPlatformProcess::BaseDir(), *FPaths::EngineContentDir(), *NewsFeedSettings->LocalSourcePath);
 
-			if (GConfig->GetString(TEXT("NewsFeed"), TEXT("EngineContentRelativeDirectory"), EngineContentRelativeDirectory, GEditorIni) && !EngineContentRelativeDirectory.IsEmpty())
+			if (!CultureString.IsEmpty())
 			{
-				FString RootDirectory = FPaths::Combine(FPlatformProcess::BaseDir(), *FPaths::EngineContentDir(), *EngineContentRelativeDirectory);
-
-				if (!CultureString.IsEmpty())
-				{
-					RootDirectory = RootDirectory.Replace(TEXT("{Culture}"), *CultureString);
-				}
-				else
-				{
-					RootDirectory = RootDirectory.Replace(TEXT("{Culture}/"), TEXT(""));
-				}
-
-				TitleFile = FLocalNewsFeedTitleFile::Create(RootDirectory);
+				RootDirectory = RootDirectory.Replace(TEXT("{Culture}"), *CultureString);
 			}
+			else
+			{
+				RootDirectory = RootDirectory.Replace(TEXT("{Culture}/"), TEXT(""));
+			}
+
+			TitleFile = FLocalNewsFeedTitleFile::Create(RootDirectory);
 		}
-		else if (SourceType == TEXT("MCP"))
-		{
-			TitleFile = Online::GetTitleFileInterface(TEXT("MCP"));
-		}
+	}
+	else if (NewsFeedSettings->Source == NEWSFEED_Mcp)
+	{
+		TitleFile = Online::GetTitleFileInterface(TEXT("MCP"));
 	}
 
+	// fetch news feed
 	if (TitleFile.IsValid())
 	{
 		TitleFile->AddOnEnumerateFilesCompleteDelegate(FOnEnumerateFilesCompleteDelegate::CreateRaw(this, &FNewsFeedCache::HandleEnumerateFilesComplete));
@@ -291,7 +287,7 @@ TSharedPtr<FSlateDynamicImageBrush> FNewsFeedCache::RawDataToBrush( FName Resour
 
 	if (bSucceeded && FSlateApplication::Get().GetRenderer()->GenerateDynamicImageResource(ResourceName, ImageWrapper->GetWidth(), ImageWrapper->GetHeight(), DecodedImage))
 	{
-		Brush = MakeShareable(new FSlateDynamicImageBrush( ResourceName, FVector2D(ImageWrapper->GetWidth(), ImageWrapper->GetHeight()) ) );
+		Brush = MakeShareable(new FSlateDynamicImageBrush(ResourceName, FVector2D(ImageWrapper->GetWidth(), ImageWrapper->GetHeight())));
 	}
 
 	return Brush;
diff --git a/Engine/Source/Editor/NewsFeed/Private/Implementation/NewsFeedSettings.h b/Engine/Source/Editor/NewsFeed/Private/Implementation/NewsFeedSettings.h
--- a/Engine/Source/Editor/NewsFeed/Private/Implementation/NewsFeedSettings.h
+++ b/Engine/Source/Editor/NewsFeed/Private/Implementation/NewsFeedSettings.h
@@ -9,12 +9,60 @@
 #include "NewsFeedSettings.generated.h"
 
 
+/**
+ * Enumerates sources for news feed data.
+ */
+UENUM()
+enum ENewsFeedSource
+{
+	/**
+	 * Fetch the news feed from the CDN.
+	 */
+	NEWSFEED_Cdn,
+
+	/**
+	 * Fetch the news feed from the local file system (for testing purposes only).
+	 */
+	NEWSFEED_Local,
+
+	/**
+	 * Fetch the news feed with MCP (not implemented yet).
+	 */
+	NEWSFEED_Mcp,
+};
+
+
+/**
+ * Holds the settings for the news feed.
+ */
 UCLASS(config=EditorGameAgnostic)
 class UNewsFeedSettings
 	: public UObject
 {
 	GENERATED_UCLASS_BODY()
 
+public:
+
+	/**
+	 * The URL at which the news feed data files are located when using the Source=Cdn.
+	 */
+	UPROPERTY(config)
+	FString CdnSourceUrl;
+
+	/**
+	 * The path to the local data files when using Source=Local.
+	 */
+	UPROPERTY(config)
+	FString LocalSourcePath;
+
+	/**
+	 * The source from which to fetch the news feed data.
+	 *
+	 * Use Local for testing, Cdn for production.
+	 */
+	UPROPERTY(config)
+	TEnumAsByte<ENewsFeedSource> Source;
+
 public:
 
 	/**
