added support for https URLs in news feed items

diff --git a/Engine/Source/Editor/NewsFeed/Private/Widgets/SNewsFeed.cpp b/Engine/Source/Editor/NewsFeed/Private/Widgets/SNewsFeed.cpp
--- a/Engine/Source/Editor/NewsFeed/Private/Widgets/SNewsFeed.cpp
+++ b/Engine/Source/Editor/NewsFeed/Private/Widgets/SNewsFeed.cpp
@@ -262,7 +262,7 @@ void SNewsFeed::HandleNewsListViewSelectionChanged( FNewsFeedItemPtr Selection,
 		{
 			bool Succeeded = false;
 
-			if (Selection->Url.StartsWith(TEXT("http://")))
+			if (Selection->Url.StartsWith(TEXT("http://")) || Selection->Url.StartsWith(TEXT("https://")))
 			{
 				FPlatformProcess::LaunchURL(*Selection->Url, NULL, NULL);
 				Succeeded = true;
