fixed News Feed crashing when spamming refresh button (TTP# 333524)

diff --git a/Engine/Source/Editor/NewsFeed/Private/Implementation/CdnNewsFeedTitleFile.cpp b/Engine/Source/Editor/NewsFeed/Private/Implementation/CdnNewsFeedTitleFile.cpp
--- a/Engine/Source/Editor/NewsFeed/Private/Implementation/CdnNewsFeedTitleFile.cpp
+++ b/Engine/Source/Editor/NewsFeed/Private/Implementation/CdnNewsFeedTitleFile.cpp
@@ -84,7 +84,7 @@ bool FCdnNewsFeedTitleFile::EnumerateFiles()
 	TSharedRef<class IHttpRequest> HttpRequest = FHttpModule::Get().CreateRequest();
 	EnumerateFilesRequests.Enqueue(&HttpRequest.Get());
 
-	HttpRequest->OnProcessRequestComplete().BindRaw(this, &FCdnNewsFeedTitleFile::EnumerateFiles_HttpRequestComplete);
+	HttpRequest->OnProcessRequestComplete().BindThreadSafeSP(this, &FCdnNewsFeedTitleFile::EnumerateFiles_HttpRequestComplete);
 	HttpRequest->SetURL( IndexUrl );
 	HttpRequest->SetVerb(TEXT("GET"));
 	HttpRequest->ProcessRequest();
@@ -245,7 +245,7 @@ bool FCdnNewsFeedTitleFile::ReadFile(const FString& FileName)
 	TSharedRef<class IHttpRequest> HttpRequest = FHttpModule::Get().CreateRequest();
 	FileRequests.Add(&HttpRequest.Get(), FPendingFileRequest(FileName));
 
-	HttpRequest->OnProcessRequestComplete().BindRaw(this, &FCdnNewsFeedTitleFile::ReadFile_HttpRequestComplete);
+	HttpRequest->OnProcessRequestComplete().BindThreadSafeSP(this, &FCdnNewsFeedTitleFile::ReadFile_HttpRequestComplete);
 	HttpRequest->SetURL( FileName );
 	HttpRequest->SetVerb(TEXT("GET"));
 
@@ -354,15 +354,23 @@ void FCdnNewsFeedTitleFile::EnumerateFiles_HttpRequestComplete(FHttpRequestPtr H
 
 void FCdnNewsFeedTitleFile::ReadFile_HttpRequestComplete(FHttpRequestPtr HttpRequest, FHttpResponsePtr HttpResponse, bool bSucceeded)
 {
+	if (!HttpRequest.IsValid())
+	{
+		return;
+	}
+
+	const FPendingFileRequest* PendingRequest = FileRequests.Find(HttpRequest.Get());
+
+	if (PendingRequest == nullptr)
+	{
+		return;
+	}
+
 	bool bResult = false;
 	FString ResponseStr, ErrorStr;
 
-	// should have a pending Http request
-	FPendingFileRequest PendingRequest = FileRequests.FindChecked(HttpRequest.Get());
-	FileRequests.Remove(HttpRequest.Get());
-
 	// Cloud file being operated on
-	FCloudFile* CloudFile = GetCloudFile(PendingRequest.FileName, true);
+	FCloudFile* CloudFile = GetCloudFile(PendingRequest->FileName, true);
 	CloudFile->AsyncState = EOnlineAsyncTaskState::Failed;
 	CloudFile->Data.Empty();
 	
@@ -372,12 +380,11 @@ void FCdnNewsFeedTitleFile::ReadFile_HttpRequestComplete(FHttpRequestPtr HttpReq
 		
 		if (EHttpResponseCodes::IsOk(HttpResponse->GetResponseCode()))
 		{
-
 			UE_LOG(LogEpicStorage, Verbose, TEXT("ReadFile request complete. url=%s code=%d"), 
 				*HttpRequest->GetURL(), HttpResponse->GetResponseCode());
 
 			// update the memory copy of the file with data that was just downloaded
-			FCloudFile* CloudFile = GetCloudFile(PendingRequest.FileName, true);
+			FCloudFile* CloudFile = GetCloudFile(PendingRequest->FileName, true);
 			CloudFile->AsyncState = EOnlineAsyncTaskState::Done;
 			CloudFile->Data = HttpResponse->GetContent();
 
@@ -396,10 +403,13 @@ void FCdnNewsFeedTitleFile::ReadFile_HttpRequestComplete(FHttpRequestPtr HttpReq
 	{
 		ErrorStr = TEXT("No response");
 	}
+
 	if (!ErrorStr.IsEmpty())
 	{
 		UE_LOG(LogEpicStorage, Warning, TEXT("ReadFile request failed. %s"), *ErrorStr);
 	}
 
-	TriggerOnReadFileCompleteDelegates(bResult, PendingRequest.FileName);
-}
\ No newline at end of file
+	TriggerOnReadFileCompleteDelegates(bResult, PendingRequest->FileName);
+
+	FileRequests.Remove(HttpRequest.Get());
+}
diff --git a/Engine/Source/Editor/NewsFeed/Private/Implementation/CdnNewsFeedTitleFile.h b/Engine/Source/Editor/NewsFeed/Private/Implementation/CdnNewsFeedTitleFile.h
--- a/Engine/Source/Editor/NewsFeed/Private/Implementation/CdnNewsFeedTitleFile.h
+++ b/Engine/Source/Editor/NewsFeed/Private/Implementation/CdnNewsFeedTitleFile.h
@@ -7,7 +7,8 @@
  * This code was salvaged from the deprecated EpicContent plug-in.
  */
 class FCdnNewsFeedTitleFile
-	: public IOnlineTitleFile
+	: public TSharedFromThis<FCdnNewsFeedTitleFile, ESPMode::ThreadSafe>
+	, public IOnlineTitleFile
 {
 	/** Info used to send request for a file */
 	struct FPendingFileRequest
