NewsFeed: moved reload throbber into correct slot; added 'pending' icon for feed icons

diff --git a/Engine/Source/Editor/EditorStyle/Private/SlateEditorStyle.cpp b/Engine/Source/Editor/EditorStyle/Private/SlateEditorStyle.cpp
--- a/Engine/Source/Editor/EditorStyle/Private/SlateEditorStyle.cpp
+++ b/Engine/Source/Editor/EditorStyle/Private/SlateEditorStyle.cpp
@@ -3578,6 +3578,7 @@ void FSlateEditorStyle::FStyle::SetupLevelEditorStyle()
 	// News Feed
 	{
 		Set( "NewsFeed.ToolbarIcon.Small", new IMAGE_BRUSH( "NewsFeed/ToolbarIcon_16x", Icon16x16 ) );
+		Set( "NewsFeed.PendingIcon", new IMAGE_BRUSH( "NewsFeed/PendingIcon", Icon16x16 ) );
 		Set( "NewsFeed.ReloadButton", new IMAGE_BRUSH( "NewsFeed/ReloadButton", Icon16x16 ) );
 		Set( "NewsFeed.SettingsButton", new IMAGE_BRUSH( "NewsFeed/SettingsButton", Icon16x16 ) );
 		Set( "NewsFeed.UnreadCountBackground", new IMAGE_BRUSH( "NewsFeed/UnreadCountBackground", Icon16x16 ) );
diff --git a/Engine/Source/Editor/NewsFeed/Private/Widgets/SNewsFeed.cpp b/Engine/Source/Editor/NewsFeed/Private/Widgets/SNewsFeed.cpp
--- a/Engine/Source/Editor/NewsFeed/Private/Widgets/SNewsFeed.cpp
+++ b/Engine/Source/Editor/NewsFeed/Private/Widgets/SNewsFeed.cpp
@@ -84,36 +84,6 @@ void SNewsFeed::Construct( const FArguments& InArgs, const FNewsFeedCacheRef& In
 			[
 				SNew(SHorizontalBox)
 
-				+ SHorizontalBox::Slot()
-					.AutoWidth()
-					.HAlign(HAlign_Center)
-					.VAlign(VAlign_Center)
-					.Padding(4.0f, 0.0f)
-					[
-						// settings button
-						SNew(SButton)
-							.ButtonStyle(FEditorStyle::Get(), "NoBorder")
-							.ContentPadding(0.0f)
-							.OnClicked(this, &SNewsFeed::HandleRefreshButtonClicked)
-							.ToolTipText(LOCTEXT("RefreshButtonToolTip", "Reload the latest news from the server."))
-							.Visibility(this, &SNewsFeed::HandleRefreshButtonVisibility)
-							[
-								SNew(SImage)
-									.Image(FEditorStyle::GetBrush("NewsFeed.ReloadButton"))
-							]
-					]
-
-				+ SHorizontalBox::Slot()
-					.FillWidth(1.0f)
-					.HAlign(HAlign_Center)
-					[
-						// mark all as read
-						SNew(SHyperlink)
-							.IsEnabled(this, &SNewsFeed::HandleMarkAllAsReadIsEnabled)
-							.OnNavigate(this, &SNewsFeed::HandleMarkAllAsReadNavigate)
-							.Text(LOCTEXT("MarkAllAsReadHyperlink", "Mark all news as read"))
-					]
-
 				+ SHorizontalBox::Slot()
 					.AutoWidth()
 					.HAlign(HAlign_Center)
@@ -126,15 +96,16 @@ void SNewsFeed::Construct( const FArguments& InArgs, const FNewsFeedCacheRef& In
 							.HAlign(HAlign_Center)
 							.VAlign(VAlign_Center)
 							[
-								// reload button
+								// settings button
 								SNew(SButton)
 									.ButtonStyle(FEditorStyle::Get(), "NoBorder")
 									.ContentPadding(0.0f)
-									.OnClicked(this, &SNewsFeed::HandleSettingsButtonClicked)
-									.ToolTipText(LOCTEXT("SettingsButtonToolTip", "News feed settings."))
+									.OnClicked(this, &SNewsFeed::HandleRefreshButtonClicked)
+									.ToolTipText(LOCTEXT("RefreshButtonToolTip", "Reload the latest news from the server."))
+									.Visibility(this, &SNewsFeed::HandleRefreshButtonVisibility)
 									[
 										SNew(SImage)
-											.Image(FEditorStyle::GetBrush("NewsFeed.SettingsButton"))
+											.Image(FEditorStyle::GetBrush("NewsFeed.ReloadButton"))
 									]
 							]
 
@@ -149,6 +120,35 @@ void SNewsFeed::Construct( const FArguments& InArgs, const FNewsFeedCacheRef& In
 									.Visibility(this, &SNewsFeed::HandleRefreshThrobberVisibility)
 							]
 					]
+
+				+ SHorizontalBox::Slot()
+					.FillWidth(1.0f)
+					.HAlign(HAlign_Center)
+					[
+						// mark all as read
+						SNew(SHyperlink)
+							.IsEnabled(this, &SNewsFeed::HandleMarkAllAsReadIsEnabled)
+							.OnNavigate(this, &SNewsFeed::HandleMarkAllAsReadNavigate)
+							.Text(LOCTEXT("MarkAllAsReadHyperlink", "Mark all news as read"))
+					]
+
+				+ SHorizontalBox::Slot()
+					.AutoWidth()
+					.HAlign(HAlign_Center)
+					.VAlign(VAlign_Center)
+					.Padding(4.0f, 0.0f)
+					[
+						// settings button
+						SNew(SButton)
+							.ButtonStyle(FEditorStyle::Get(), "NoBorder")
+							.ContentPadding(0.0f)
+							.OnClicked(this, &SNewsFeed::HandleSettingsButtonClicked)
+							.ToolTipText(LOCTEXT("SettingsButtonToolTip", "News feed settings."))	
+							[
+								SNew(SImage)
+									.Image(FEditorStyle::GetBrush("NewsFeed.SettingsButton"))
+							]	
+					]
 			]
 	];
 
diff --git a/Engine/Source/Editor/NewsFeed/Private/Widgets/SNewsFeedListRow.h b/Engine/Source/Editor/NewsFeed/Private/Widgets/SNewsFeedListRow.h
--- a/Engine/Source/Editor/NewsFeed/Private/Widgets/SNewsFeedListRow.h
+++ b/Engine/Source/Editor/NewsFeed/Private/Widgets/SNewsFeedListRow.h
@@ -58,8 +58,15 @@ public:
 				.VAlign(VAlign_Top)
 				.Padding(4.0f)
 				[
-					SNew(SImage)
-						.Image(this, &SNewsFeedListRow::HandleIconImage)
+					SNew(SBox)
+						.HAlign(HAlign_Fill)
+						.VAlign(VAlign_Fill)
+						.HeightOverride(40.0f)
+						.WidthOverride(40.0f)
+						[
+							SNew(SImage)
+								.Image(this, &SNewsFeedListRow::HandleIconImage)
+						]
 				];
 		}
 		else if (ColumnName == TEXT("Content"))
