NewsFeed: fixed news feed not remembering read items (TTP# 333525)

diff --git a/Engine/Config/BaseEditorGameAgnostic.ini b/Engine/Config/BaseEditorGameAgnostic.ini
--- a/Engine/Config/BaseEditorGameAgnostic.ini
+++ b/Engine/Config/BaseEditorGameAgnostic.ini
@@ -1,8 +1,8 @@
 [/Script/UnrealEd.EditorGameAgnosticSettings]
 
 [/Script/NewsFeed.NewsFeedSettings]
-CdnSourceUrl="http://cdn.unrealengine.com/newsfeed/{Culture}/index.json"
-LocalSourcePath="Editor/NewsFeed/{Culture}/"
+CdnSourceUrl="http://cdn.unrealengine.com/newsfeed/[[Culture]]/index.json"
+LocalSourcePath="Editor/NewsFeed/[[Culture]]/"
 Source=NEWSFEED_Cdn
 MaxItemsToShow=10
-ShowOnlyUnreadItems=True
\ No newline at end of file
+ShowOnlyUnreadItems=True
diff --git a/Engine/Source/Editor/NewsFeed/Private/Implementation/NewsFeedCache.cpp b/Engine/Source/Editor/NewsFeed/Private/Implementation/NewsFeedCache.cpp
--- a/Engine/Source/Editor/NewsFeed/Private/Implementation/NewsFeedCache.cpp
+++ b/Engine/Source/Editor/NewsFeed/Private/Implementation/NewsFeedCache.cpp
@@ -111,11 +111,11 @@ void FNewsFeedCache::LoadTitleFile( )
 
 			if (!CultureString.IsEmpty())
 			{
-				CdnSourceUrl = NewsFeedSettings->CdnSourceUrl.Replace(TEXT("{Culture}"), *CultureString);
+				CdnSourceUrl = NewsFeedSettings->CdnSourceUrl.Replace(TEXT("[[Culture]]"), *CultureString);
 			}
 			else
 			{
-				CdnSourceUrl = NewsFeedSettings->CdnSourceUrl.Replace(TEXT("{Culture}/"), TEXT(""));
+				CdnSourceUrl = NewsFeedSettings->CdnSourceUrl.Replace(TEXT("[[Culture]]/"), TEXT(""));
 			}
 
 			TitleFile = FCdnNewsFeedTitleFile::Create(CdnSourceUrl);
@@ -129,11 +129,11 @@ void FNewsFeedCache::LoadTitleFile( )
 
 			if (!CultureString.IsEmpty())
 			{
-				RootDirectory = RootDirectory.Replace(TEXT("{Culture}"), *CultureString);
+				RootDirectory = RootDirectory.Replace(TEXT("[[Culture]]"), *CultureString);
 			}
 			else
 			{
-				RootDirectory = RootDirectory.Replace(TEXT("{Culture}/"), TEXT(""));
+				RootDirectory = RootDirectory.Replace(TEXT("[[Culture]]/"), TEXT(""));
 			}
 
 			TitleFile = FLocalNewsFeedTitleFile::Create(RootDirectory);
diff --git a/Engine/Source/Editor/NewsFeed/Private/Implementation/NewsFeedSettings.h b/Engine/Source/Editor/NewsFeed/Private/Implementation/NewsFeedSettings.h
--- a/Engine/Source/Editor/NewsFeed/Private/Implementation/NewsFeedSettings.h
+++ b/Engine/Source/Editor/NewsFeed/Private/Implementation/NewsFeedSettings.h
@@ -58,7 +58,7 @@ public:
 	/**
 	 * The source from which to fetch the news feed data.
 	 *
-	 * Use Local for testing, Cdn for production.
+	 * Use Local for testing, NEWSFEED_Cdn for production.
 	 */
 	UPROPERTY(config)
 	TEnumAsByte<ENewsFeedSource> Source;
diff --git a/Engine/Source/Editor/NewsFeed/Private/Widgets/SNewsFeed.cpp b/Engine/Source/Editor/NewsFeed/Private/Widgets/SNewsFeed.cpp
--- a/Engine/Source/Editor/NewsFeed/Private/Widgets/SNewsFeed.cpp
+++ b/Engine/Source/Editor/NewsFeed/Private/Widgets/SNewsFeed.cpp
@@ -211,12 +211,17 @@ bool SNewsFeed::HandleMarkAllAsReadIsEnabled( ) const
 
 void SNewsFeed::HandleMarkAllAsReadNavigate( )
 {
+	UNewsFeedSettings* NewsFeedSettings = GetMutableDefault<UNewsFeedSettings>();
+
 	for (auto NewsFeedItem : NewsFeedItemList)
 	{
+		NewsFeedSettings->ReadItems.AddUnique(NewsFeedItem->ItemId);
 		NewsFeedItem->Read = true;
 	}
 
-	if (GetDefault<UNewsFeedSettings>()->ShowOnlyUnreadItems)
+	NewsFeedSettings->SaveConfig();
+
+	if (NewsFeedSettings->ShowOnlyUnreadItems)
 	{
 		ReloadNewsFeedItems();
 	}
@@ -253,9 +258,9 @@ void SNewsFeed::HandleNewsListViewSelectionChanged( FNewsFeedItemPtr Selection,
 	if (Selection.IsValid() && (SelectInfo == ESelectInfo::OnMouseClick))
 	{
 		TArray<FAnalyticsEventAttribute> EventAttributes;
-
 		EventAttributes.Add(FAnalyticsEventAttribute(TEXT("ItemID"), Selection->ItemId.ToString()));
 
+		// execute item action
 		IDesktopPlatform* DesktopPlatform = FDesktopPlatformModule::Get();
 
 		if (DesktopPlatform != nullptr)
@@ -284,12 +289,16 @@ void SNewsFeed::HandleNewsListViewSelectionChanged( FNewsFeedItemPtr Selection,
 			}
 		}
 
-		// hack: mark as read
+		// mark item as read
+		UNewsFeedSettings* NewsFeedSettings = GetMutableDefault<UNewsFeedSettings>();
+		NewsFeedSettings->ReadItems.AddUnique(Selection->ItemId);
+		NewsFeedSettings->SaveConfig();
 		Selection->Read = true;
-		
+
+		// refresh list view
 		NewsFeedItemListView->SetSelection(nullptr);
 
-		if (GetDefault<UNewsFeedSettings>()->ShowOnlyUnreadItems)
+		if (NewsFeedSettings->ShowOnlyUnreadItems)
 		{
 			ReloadNewsFeedItems();
 		}
