NewsFeed: renamed setting for number of items to show

diff --git a/Engine/Source/Editor/NewsFeed/Private/Implementation/NewsFeedSettings.h b/Engine/Source/Editor/NewsFeed/Private/Implementation/NewsFeedSettings.h
--- a/Engine/Source/Editor/NewsFeed/Private/Implementation/NewsFeedSettings.h
+++ b/Engine/Source/Editor/NewsFeed/Private/Implementation/NewsFeedSettings.h
@@ -68,7 +68,7 @@ public:
 	/**
 	 * The maximum number of news items to show.
 	 */
-	UPROPERTY(config, EditAnywhere, Category=Display)
+	UPROPERTY(config, EditAnywhere, Category=Display, meta=(DisplayName="Number Of Most Recent Items To Show"))
 	int32 MaxItemsToShow;
 
 	/**
