NewsFeed: Added an option to show & hide older news feed items (TTP# 334454)

diff --git a/Engine/Source/Editor/EditorStyle/Private/SlateEditorStyle.cpp b/Engine/Source/Editor/EditorStyle/Private/SlateEditorStyle.cpp
--- a/Engine/Source/Editor/EditorStyle/Private/SlateEditorStyle.cpp
+++ b/Engine/Source/Editor/EditorStyle/Private/SlateEditorStyle.cpp
@@ -3583,6 +3583,7 @@ void FSlateEditorStyle::FStyle::SetupLevelEditorStyle()
 	// News Feed
 	{
 		Set( "NewsFeed.ToolbarIcon.Small", new IMAGE_BRUSH( "NewsFeed/ToolbarIcon_16x", Icon16x16 ) );
+		Set( "NewsFeed.MarkAsRead", new IMAGE_BRUSH( "NewsFeed/MarkAsRead", Icon16x16 ) );
 		Set( "NewsFeed.PendingIcon", new IMAGE_BRUSH( "NewsFeed/PendingIcon", Icon16x16 ) );
 		Set( "NewsFeed.ReloadButton", new IMAGE_BRUSH( "NewsFeed/ReloadButton", Icon16x16 ) );
 		Set( "NewsFeed.SettingsButton", new IMAGE_BRUSH( "NewsFeed/SettingsButton", Icon16x16 ) );
diff --git a/Engine/Source/Editor/NewsFeed/Private/Widgets/SNewsFeed.cpp b/Engine/Source/Editor/NewsFeed/Private/Widgets/SNewsFeed.cpp
--- a/Engine/Source/Editor/NewsFeed/Private/Widgets/SNewsFeed.cpp
+++ b/Engine/Source/Editor/NewsFeed/Private/Widgets/SNewsFeed.cpp
@@ -28,6 +28,8 @@ SNewsFeed::~SNewsFeed( )
 void SNewsFeed::Construct( const FArguments& InArgs, const FNewsFeedCacheRef& InNewsFeedCache )
 {
 	NewsFeedCache = InNewsFeedCache;
+	HiddenNewsItemCount = 0;
+	ShowingOlderNews = false;
 
 	ChildSlot
 	[
@@ -121,20 +123,42 @@ void SNewsFeed::Construct( const FArguments& InArgs, const FNewsFeedCacheRef& In
 							]
 					]
 
+				+ SHorizontalBox::Slot()
+					.AutoWidth()
+					.HAlign(HAlign_Left)
+					.VAlign(VAlign_Center)
+					.Padding(4.0f, 0.0f)
+					[
+						// mark as read button
+						SNew(SButton)
+							.ButtonStyle(FEditorStyle::Get(), "NoBorder")
+							.ContentPadding(0.0f)
+							.IsEnabled(this, &SNewsFeed::HandleMarkAllAsReadIsEnabled)
+							.OnClicked(this, &SNewsFeed::HandleMarkAllAsReadClicked)
+							.ToolTipText(LOCTEXT("MarkAsReadButtonToolTip", "Mark all news as read."))	
+							[
+								SNew(SImage)
+									.Image(FEditorStyle::GetBrush("NewsFeed.MarkAsRead"))
+							]
+					]
+
 				+ SHorizontalBox::Slot()
 					.FillWidth(1.0f)
 					.HAlign(HAlign_Center)
+					.VAlign(VAlign_Center)
+					.Padding(4.0f, 0.0f)
 					[
-						// mark all as read
+						// show older news hyper-link
 						SNew(SHyperlink)
-							.IsEnabled(this, &SNewsFeed::HandleMarkAllAsReadIsEnabled)
-							.OnNavigate(this, &SNewsFeed::HandleMarkAllAsReadNavigate)
-							.Text(LOCTEXT("MarkAllAsReadHyperlink", "Mark all news as read"))
+							.OnNavigate(this, &SNewsFeed::HandleShowOlderNewsNavigate)
+							.Text(this, &SNewsFeed::HandleShowOlderNewsText)
+							.ToolTipText(LOCTEXT("ShowOlderNewsToolTip", "Toggle visibility of news items that are below your threshold for recent news."))
+							.Visibility(this, &SNewsFeed::HandleShowOlderNewsVisibility)
 					]
 
 				+ SHorizontalBox::Slot()
 					.AutoWidth()
-					.HAlign(HAlign_Center)
+					.HAlign(HAlign_Right)
 					.VAlign(VAlign_Center)
 					.Padding(4.0f, 0.0f)
 					[
@@ -147,7 +171,7 @@ void SNewsFeed::Construct( const FArguments& InArgs, const FNewsFeedCacheRef& In
 							[
 								SNew(SImage)
 									.Image(FEditorStyle::GetBrush("NewsFeed.SettingsButton"))
-							]	
+							]
 					]
 			]
 	];
@@ -166,18 +190,22 @@ void SNewsFeed::Construct( const FArguments& InArgs, const FNewsFeedCacheRef& In
 void SNewsFeed::ReloadNewsFeedItems( )
 {
 	NewsFeedItemList.Empty();
+	HiddenNewsItemCount = 0;
 
 	const UNewsFeedSettings* NewsFeedSettings = GetDefault<UNewsFeedSettings>();
+	const TArray<FNewsFeedItemPtr>& NewsFeedItems = NewsFeedCache->GetItems();
 
-	for (auto NewsFeedItem : NewsFeedCache->GetItems())
+	for (auto NewsFeedItem : NewsFeedItems)
 	{
 		if (!GetDefault<UNewsFeedSettings>()->ShowOnlyUnreadItems || !NewsFeedItem->Read)
 		{
-			NewsFeedItemList.Add(NewsFeedItem);
-
-			if (NewsFeedItemList.Num() == NewsFeedSettings->MaxItemsToShow)
+			if (ShowingOlderNews || (NewsFeedItemList.Num() < NewsFeedSettings->MaxItemsToShow))
+			{
+				NewsFeedItemList.Add(NewsFeedItem);
+			}
+			else
 			{
-				break;
+				++HiddenNewsItemCount;
 			}
 		}
 	}
@@ -195,21 +223,7 @@ EVisibility SNewsFeed::HandleLoadFailedBoxVisibility( ) const
 }
 
 
-bool SNewsFeed::HandleMarkAllAsReadIsEnabled( ) const
-{
-	for (auto NewsFeedItem : NewsFeedItemList)
-	{
-		if (!NewsFeedItem->Read)
-		{
-			return true;
-		}
-	}
-
-	return false;
-}
-
-
-void SNewsFeed::HandleMarkAllAsReadNavigate( )
+FReply SNewsFeed::HandleMarkAllAsReadClicked( )
 {
 	UNewsFeedSettings* NewsFeedSettings = GetMutableDefault<UNewsFeedSettings>();
 
@@ -225,6 +239,22 @@ void SNewsFeed::HandleMarkAllAsReadNavigate( )
 	{
 		ReloadNewsFeedItems();
 	}
+
+	return FReply::Handled();
+}
+
+
+bool SNewsFeed::HandleMarkAllAsReadIsEnabled( ) const
+{
+	for (auto NewsFeedItem : NewsFeedItemList)
+	{
+		if (!NewsFeedItem->Read)
+		{
+			return true;
+		}
+	}
+
+	return false;
 }
 
 
@@ -242,6 +272,7 @@ void SNewsFeed::HandleNewsFeedLoaderLoadFailed( )
 
 void SNewsFeed::HandleNewsFeedSettingsChanged( FName PropertyName )
 {
+	ShowingOlderNews = false;
 	ReloadNewsFeedItems();
 }
 
@@ -357,4 +388,28 @@ FReply SNewsFeed::HandleSettingsButtonClicked( ) const
 }
 
 
+void SNewsFeed::HandleShowOlderNewsNavigate( )
+{
+	ShowingOlderNews = !ShowingOlderNews;
+	ReloadNewsFeedItems();
+}
+
+
+FText SNewsFeed::HandleShowOlderNewsText( ) const
+{
+	if (ShowingOlderNews)
+	{
+		return LOCTEXT("HideOlderNewsHyperlink", "Hide older news");
+	}
+
+	return FText::Format(LOCTEXT("ShowOlderNewsHyperlink", "Show older news ({0})"), FText::AsNumber(HiddenNewsItemCount));
+}
+
+
+EVisibility SNewsFeed::HandleShowOlderNewsVisibility( ) const
+{
+	return (ShowingOlderNews || (HiddenNewsItemCount > 0)) ? EVisibility::Visible : EVisibility::Hidden;
+}
+
+
 #undef LOCTEXT_NAMESPACE
diff --git a/Engine/Source/Editor/NewsFeed/Private/Widgets/SNewsFeed.h b/Engine/Source/Editor/NewsFeed/Private/Widgets/SNewsFeed.h
--- a/Engine/Source/Editor/NewsFeed/Private/Widgets/SNewsFeed.h
+++ b/Engine/Source/Editor/NewsFeed/Private/Widgets/SNewsFeed.h
@@ -47,11 +47,11 @@ private:
 	// Callback for getting the visibility of the 'Loading failed' box.
 	EVisibility HandleLoadFailedBoxVisibility( ) const;
 
-	// Callback for getting the enabled state of the 'Mark all as read' hyper-link.
-	bool HandleMarkAllAsReadIsEnabled( ) const;
+	// Callback for clicking the 'Mark all as read' button.
+	FReply HandleMarkAllAsReadClicked( );
 
-	// Callback for clicking the 'Mark all as read' hyper-link.
-	void HandleMarkAllAsReadNavigate( );
+	// Callback for getting the enabled state of the 'Mark all as read' button.
+	bool HandleMarkAllAsReadIsEnabled( ) const;
 
 	// Callback for completed news feeds loads.
 	void HandleNewsFeedLoaderLoadCompleted( );
@@ -89,8 +89,20 @@ private:
 	// Callback for clicking the 'Settings' button.
 	FReply HandleSettingsButtonClicked( ) const;
 
+	// Callback for clicking the 'Show older news' hyper-link.
+	void HandleShowOlderNewsNavigate( );
+
+	// Callback for getting the text of the 'Show older news' hyper-link.
+	FText HandleShowOlderNewsText( ) const;
+
+	// Callback for getting the visibility of the 'Show older news' hyper-link.
+	EVisibility HandleShowOlderNewsVisibility( ) const;
+
 private:
 
+	// Holds the number of hidden news items.
+	int32 HiddenNewsItemCount;
+
 	// Holds a pointer the news feed cache.
 	FNewsFeedCachePtr NewsFeedCache;
 
@@ -99,4 +111,7 @@ private:
 
 	// Holds the list view for news items.
 	TSharedPtr<SListView<FNewsFeedItemPtr> > NewsFeedItemListView;
+
+	// Whether older news are currently being shown as well.
+	bool ShowingOlderNews;
 };
