NewsFeed: don't save configs if running on build machine

diff --git a/Engine/Source/Editor/NewsFeed/Private/Implementation/NewsFeedCache.cpp b/Engine/Source/Editor/NewsFeed/Private/Implementation/NewsFeedCache.cpp
--- a/Engine/Source/Editor/NewsFeed/Private/Implementation/NewsFeedCache.cpp
+++ b/Engine/Source/Editor/NewsFeed/Private/Implementation/NewsFeedCache.cpp
@@ -307,7 +307,11 @@ void FNewsFeedCache::SaveSettings( )
 		}
 	}
 
-	NewsFeedSettings->SaveConfig();
+	// Save config to file, but only if we are not the build machine since settings may put the builder in an unclean state
+	if (!GIsBuildMachine)
+	{
+		NewsFeedSettings->SaveConfig();
+	}
 }
 
 
