NewsFeed: Documentation and formatting cleanup pass

diff --git a/Engine/Source/Editor/NewsFeed/Private/Implementation/NewsFeedCache.h b/Engine/Source/Editor/NewsFeed/Private/Implementation/NewsFeedCache.h
--- a/Engine/Source/Editor/NewsFeed/Private/Implementation/NewsFeedCache.h
+++ b/Engine/Source/Editor/NewsFeed/Private/Implementation/NewsFeedCache.h
@@ -1,6 +1,8 @@
 #pragma once
 
 
+/**
+ */
 enum class ENewsFeedCultureSpec
 {
 	Full = 0,
@@ -9,29 +11,38 @@ enum class ENewsFeedCultureSpec
 };
 
 
+/**
+ * Enumerates valid news feed cache states.
+ */
 namespace ENewsFeedState
 {
 	enum Type
 	{
+		/** Download of news feed has not started yet. */
 		NotStarted = 0,
+
+		/** News feed is being fetched from data source. */
 		Loading = 1,
+
+		/** News feed finished loading successfully. */
 		Success = 2,
+
+		/** News feed failed to load. */
 		Failure = 3
 	};
 }
 
 
-/**
- * Type definition for shared pointers to instances of FNewsFeedCache.
- */
+/** Type definition for shared pointers to instances of FNewsFeedCache. */
 typedef TSharedPtr<class FNewsFeedCache> FNewsFeedCachePtr;
 
-/**
- * Type definition for shared references to instances of FNewsFeedCache.
- */
+/** Type definition for shared references to instances of FNewsFeedCache. */
 typedef TSharedRef<class FNewsFeedCache> FNewsFeedCacheRef;
 
 
+/**
+ * Implements a cache of news feed items that were fetched from the data source.
+ */
 class FNewsFeedCache
 {
 public:
@@ -55,7 +66,6 @@ public:
 	 * If the icon failed to load, an error icon will be returned.
 	 *
 	 * @param IconName The name of the icon to get.
-	 *
 	 * @return The icon brush.
 	 */
 	const FSlateBrush* GetIcon( const FString& IconName ) const;
@@ -74,7 +84,6 @@ public:
 	 * Checks whether the last attempt to load a feed has finished.
 	 *
 	 * @return true if feed loading finished, false otherwise.
-	 *
 	 * @see IsLoading
 	 */
 	bool IsFinished( ) const
@@ -86,7 +95,6 @@ public:
 	 * Checks whether this loader is currently loading a news feed.
 	 *
 	 * @return true if loading, false otherwise.
-	 *
 	 * @see IsFinished
 	 */
 	bool IsLoading( ) const
@@ -98,7 +106,6 @@ public:
 	 * Loads the news feed from the given title file.
 	 *
 	 * @param InTitleFile The title file to read the news feed from.
-	 *
 	 * @return true if the news feed is being loaded, false otherwise.
 	 */
 	//bool LoadFeed( IOnlineTitleFilePtr InTitleFile );
@@ -157,7 +164,6 @@ protected:
 	 *
 	 * @param ResourceName 
 	 * @param RawData The raw binary data to convert.
-	 *
 	 * @return The Slate brush.
 	 */
 	TSharedPtr<FSlateDynamicImageBrush> RawDataToBrush( FName ResourceName, const TArray< uint8 >& RawData ) const;
diff --git a/Engine/Source/Editor/NewsFeed/Private/Implementation/NewsFeedSettings.h b/Engine/Source/Editor/NewsFeed/Private/Implementation/NewsFeedSettings.h
--- a/Engine/Source/Editor/NewsFeed/Private/Implementation/NewsFeedSettings.h
+++ b/Engine/Source/Editor/NewsFeed/Private/Implementation/NewsFeedSettings.h
@@ -15,19 +15,13 @@
 UENUM()
 enum ENewsFeedSource
 {
-	/**
-	 * Fetch the news feed from the CDN.
-	 */
+	/** Fetch the news feed from the CDN. */
 	NEWSFEED_Cdn,
 
-	/**
-	 * Fetch the news feed from the local file system (for testing purposes only).
-	 */
+	/** Fetch the news feed from the local file system (for testing purposes only). */
 	NEWSFEED_Local,
 
-	/**
-	 * Fetch the news feed with MCP (not implemented yet).
-	 */
+	/** Fetch the news feed with MCP (not implemented yet). */
 	NEWSFEED_Mcp,
 };
 
diff --git a/Engine/Source/Editor/NewsFeed/Private/Models/NewsFeedItem.h b/Engine/Source/Editor/NewsFeed/Private/Models/NewsFeedItem.h
--- a/Engine/Source/Editor/NewsFeed/Private/Models/NewsFeedItem.h
+++ b/Engine/Source/Editor/NewsFeed/Private/Models/NewsFeedItem.h
@@ -7,14 +7,10 @@
 #pragma once
 
 
-/**
- * Type definition for shared pointers to instances of FNewsFeedItem.
- */
+/** Type definition for shared pointers to instances of FNewsFeedItem. */
 typedef TSharedPtr<struct FNewsFeedItem> FNewsFeedItemPtr;
 
-/**
- * Type definition for shared references to instances of FNewsFeedItem.
- */
+/** Type definition for shared references to instances of FNewsFeedItem. */
 typedef TSharedRef<struct FNewsFeedItem> FNewsFeedItemRef;
 
 
@@ -23,46 +19,32 @@ typedef TSharedRef<struct FNewsFeedItem> FNewsFeedItemRef;
  */
 struct FNewsFeedItem
 {
-	/**
-	 * The news excerpt.
-	 */
+	/** The news excerpt. */
 	FText Excerpt;
 
-	/**
-	 * The full news text.
-	 */
+	/** The full news text. */
 	FText FullText;
 
-	/**
-	 * Holds the name of the icon.
-	 */
+	/** Holds the name of the icon. */
 	FString IconName;
 
-	/**
-	 * The date and time at which the news item was issued.
-	 */
+	/** The date and time at which the news item was issued. */
 	FDateTime Issued;
 
-	/**
-	 * The news item's unique identifier.
-	 */
+	/** The news item's unique identifier. */
 	FGuid ItemId;
 
-	/**
-	 * Whether this news item has been marked as read.
-	 */
+	/** Whether this news item has been marked as read. */
 	bool Read;
 
-	/**
-	 * The news title.
-	 */
+	/** The news title. */
 	FText Title;
 
-	/**
-	 * The URL to the full news release.
-	 */
+	/** The URL to the full news release. */
 	FString Url;
 
+public:
+
 	/**
 	 * Creates and initializes a new item with the specified identifier.
 	 *
diff --git a/Engine/Source/Editor/NewsFeed/Private/Widgets/SNewsFeed.h b/Engine/Source/Editor/NewsFeed/Private/Widgets/SNewsFeed.h
--- a/Engine/Source/Editor/NewsFeed/Private/Widgets/SNewsFeed.h
+++ b/Engine/Source/Editor/NewsFeed/Private/Widgets/SNewsFeed.h
@@ -30,8 +30,8 @@ public:
 	/**
 	 * Constructs the widget.
 	 *
-	 * @param InArgs - The construction arguments.
-	 * @param InNewsFeedCache - The news feed cache to use.
+	 * @param InArgs The construction arguments.
+	 * @param InNewsFeedCache The news feed cache to use.
 	 */
 	void Construct( const FArguments& InArgs, const FNewsFeedCacheRef& InNewsFeedCache );
 
@@ -110,7 +110,7 @@ private:
 	TArray<FNewsFeedItemPtr> NewsFeedItemList;
 
 	// Holds the list view for news items.
-	TSharedPtr<SListView<FNewsFeedItemPtr> > NewsFeedItemListView;
+	TSharedPtr<SListView<FNewsFeedItemPtr>> NewsFeedItemListView;
 
 	// Whether older news are currently being shown as well.
 	bool ShowingOlderNews;
diff --git a/Engine/Source/Editor/NewsFeed/Private/Widgets/SNewsFeedListRow.h b/Engine/Source/Editor/NewsFeed/Private/Widgets/SNewsFeedListRow.h
--- a/Engine/Source/Editor/NewsFeed/Private/Widgets/SNewsFeedListRow.h
+++ b/Engine/Source/Editor/NewsFeed/Private/Widgets/SNewsFeedListRow.h
@@ -27,7 +27,7 @@ public:
 	/**
 	 * Constructs the widget.
 	 *
-	 * @param InArgs - The arguments.
+	 * @param InArgs The arguments.
 	 */
 	void Construct( const FArguments& InArgs, const TSharedRef<STableViewBase>& InOwnerTableView, const FNewsFeedCacheRef& InNewsFeedCache )
 	{
@@ -44,8 +44,7 @@ public:
 	/**
 	 * Generates the widget for the specified column.
 	 *
-	 * @param ColumnName - The name of the column to generate the widget for.
-	 *
+	 * @param ColumnName The name of the column to generate the widget for.
 	 * @return The widget.
 	 */
 	BEGIN_SLATE_FUNCTION_BUILD_OPTIMIZATION
